[Terminal] Move `Terminal` plugin to `Engine/Plugins/Experimental`.

#Synced-CL 52520731

diff --git a/Engine/Plugins/Experimental/Terminal/Source/Terminal/Private/ConPTYSession.cpp b/Engine/Plugins/Experimental/Terminal/Source/Terminal/Private/ConPTYSession.cpp
new file mode 100644
--- /dev/null
+++ b/Engine/Plugins/Experimental/Terminal/Source/Terminal/Private/ConPTYSession.cpp
@@ -0,0 +1,542 @@
+// Copyright Epic Games, Inc. All Rights Reserved.
+
+#include "ConPTYSession.h"
+
+#if PLATFORM_WINDOWS
+
+#include "Async/Async.h"
+#include "Interfaces/IPluginManager.h"
+#include "Misc/Paths.h"
+
+#include "Windows/AllowWindowsPlatformTypes.h"
+#include <windows.h>
+#include "Windows/HideWindowsPlatformTypes.h"
+
+namespace
+{
+
+struct FConPTYAPI
+{
+	typedef HRESULT(WINAPI* CreatePseudoConsoleFunc)(COORD Size, HANDLE hInput, HANDLE hOutput, DWORD dwFlags, void** phPC);
+	typedef HRESULT(WINAPI* ResizePseudoConsoleFunc)(void* hPC, COORD Size);
+	typedef void(WINAPI* ClosePseudoConsoleFunc)(void* hPC);
+
+	CreatePseudoConsoleFunc CreatePseudoConsole = nullptr;
+	ResizePseudoConsoleFunc ResizePseudoConsole = nullptr;
+	ClosePseudoConsoleFunc ClosePseudoConsole = nullptr;
+
+	/** Non-null if we loaded a bundled conpty.dll. */
+	HMODULE BundledConPTYModule = nullptr;
+
+	~FConPTYAPI()
+	{
+		if (BundledConPTYModule)
+		{
+			FreeLibrary(BundledConPTYModule);
+			BundledConPTYModule = nullptr;
+		}
+	}
+};
+
+FConPTYAPI InitConPTYAPI()
+{
+	FConPTYAPI API;
+
+	// Try bundled conpty.dll first (ships with OpenConsole.exe, bypasses system conhost.exe).
+	HMODULE SourceModule = nullptr;
+	TSharedPtr<IPlugin> TerminalPlugin = IPluginManager::Get().FindPlugin(TEXT("Terminal"));
+	if (TerminalPlugin.IsValid())
+	{
+		const FString BundledConPTYPath = FPaths::ConvertRelativePathToFull(FPaths::Combine(
+			TerminalPlugin->GetBaseDir(),
+			TEXT("Binaries"), FPlatformProcess::GetBinariesSubdirectory(),
+			TEXT("conpty.dll")));
+
+		if (FPaths::FileExists(BundledConPTYPath))
+		{
+			SourceModule = LoadLibraryExW(*BundledConPTYPath, nullptr, LOAD_LIBRARY_SEARCH_DLL_LOAD_DIR | LOAD_LIBRARY_SEARCH_SYSTEM32);
+			if (SourceModule)
+			{
+				API.BundledConPTYModule = SourceModule;
+				UE_LOG(LogTerminal, Display, TEXT("Loaded bundled ConPTY: %s"), *BundledConPTYPath);
+			}
+			else
+			{
+				UE_LOG(LogTerminal, Warning, TEXT("Failed to load bundled conpty.dll from %s (error %d). Falling back to system ConPTY."), *BundledConPTYPath, GetLastError());
+			}
+		}
+	}
+
+	if (!SourceModule)
+	{
+		SourceModule = GetModuleHandleW(L"kernel32.dll");
+		UE_LOG(LogTerminal, Display, TEXT("Using system ConPTY from kernel32.dll."));
+	}
+
+	if (SourceModule)
+	{
+#pragma warning(push)
+#pragma warning(disable: 4191) // Unsafe conversion from FARPROC to specific function pointer type.
+		API.CreatePseudoConsole = reinterpret_cast<FConPTYAPI::CreatePseudoConsoleFunc>(
+			GetProcAddress(SourceModule, "CreatePseudoConsole"));
+		API.ResizePseudoConsole = reinterpret_cast<FConPTYAPI::ResizePseudoConsoleFunc>(
+			GetProcAddress(SourceModule, "ResizePseudoConsole"));
+		API.ClosePseudoConsole = reinterpret_cast<FConPTYAPI::ClosePseudoConsoleFunc>(
+			GetProcAddress(SourceModule, "ClosePseudoConsole"));
+#pragma warning(pop)
+	}
+	return API;
+}
+
+// C++11 magic statics guarantee thread-safe one-time initialization.
+FConPTYAPI& GetConPTYAPI()
+{
+	static FConPTYAPI API = InitConPTYAPI();
+	return API;
+}
+
+} // anonymous namespace
+
+FConPTYSession::FConPTYSession()
+{
+}
+
+FConPTYSession::~FConPTYSession()
+{
+	Shutdown();
+}
+
+bool FConPTYSession::IsConPTYAvailable()
+{
+	const FConPTYAPI& API = GetConPTYAPI();
+	return API.CreatePseudoConsole && API.ResizePseudoConsole && API.ClosePseudoConsole;
+}
+
+bool FConPTYSession::Create(const FString& ShellPath, const FString& WorkingDirectory, int32 Columns, int32 Rows)
+{
+	const FConPTYAPI& API = GetConPTYAPI();
+	if (!API.CreatePseudoConsole)
+	{
+		UE_LOG(LogTerminal, Error, TEXT("ConPTY API is not available on this system."));
+		return false;
+	}
+
+	// Create the pipes.
+	HANDLE InputPipeReadHandle = INVALID_HANDLE_VALUE;
+	HANDLE InputPipeWriteHandle = INVALID_HANDLE_VALUE;
+	HANDLE OutputPipeReadHandle = INVALID_HANDLE_VALUE;
+	HANDLE OutputPipeWriteHandle = INVALID_HANDLE_VALUE;
+
+	if (!CreatePipe(&InputPipeReadHandle, &InputPipeWriteHandle, nullptr, 0))
+	{
+		UE_LOG(LogTerminal, Error, TEXT("Failed to create input pipe."));
+		return false;
+	}
+	if (!CreatePipe(&OutputPipeReadHandle, &OutputPipeWriteHandle, nullptr, 0))
+	{
+		UE_LOG(LogTerminal, Error, TEXT("Failed to create output pipe."));
+		CloseHandle(InputPipeReadHandle);
+		CloseHandle(InputPipeWriteHandle);
+		return false;
+	}
+
+	// Create the pseudo-console.
+	const int32 ClampedColumns = FMath::Clamp(Columns, 1, static_cast<int32>(SHRT_MAX));
+	const int32 ClampedRows = FMath::Clamp(Rows, 1, static_cast<int32>(SHRT_MAX));
+	COORD ConsoleSize = {};
+	ConsoleSize.X = static_cast<SHORT>(ClampedColumns);
+	ConsoleSize.Y = static_cast<SHORT>(ClampedRows);
+
+	void* PseudoConsole = nullptr;
+	HRESULT Result = API.CreatePseudoConsole(ConsoleSize, InputPipeReadHandle, OutputPipeWriteHandle, 0, &PseudoConsole);
+	if (FAILED(Result))
+	{
+		UE_LOG(LogTerminal, Error, TEXT("CreatePseudoConsole failed with HRESULT 0x%08X."), Result);
+		CloseHandle(InputPipeReadHandle);
+		CloseHandle(InputPipeWriteHandle);
+		CloseHandle(OutputPipeReadHandle);
+		CloseHandle(OutputPipeWriteHandle);
+		return false;
+	}
+
+	// Spawn the shell process.
+	FString EffectiveShellPath = ShellPath;
+	if (EffectiveShellPath.IsEmpty())
+	{
+		// Use COMSPEC environment variable as default.
+		TCHAR ComSpec[MAX_PATH];
+		if (GetEnvironmentVariableW(TEXT("COMSPEC"), ComSpec, MAX_PATH) > 0)
+		{
+			EffectiveShellPath = ComSpec;
+		}
+		else
+		{
+			EffectiveShellPath = TEXT("cmd.exe");
+		}
+	}
+
+	UE_LOG(LogTerminal, VeryVerbose, TEXT("Resolved shell path: %s (requested: %s). Console size: %dx%d."), *EffectiveShellPath, ShellPath.IsEmpty() ? TEXT("<default>") : *ShellPath, ClampedColumns, ClampedRows);
+
+	// Set up the startup info with the pseudo-console attribute.
+	// First call with nullptr intentionally fails to retrieve the required buffer size.
+	SIZE_T AttributeListSize = 0;
+	(void)InitializeProcThreadAttributeList(nullptr, 1, 0, &AttributeListSize);
+	TArray<uint8> AttributeListStorage;
+	AttributeListStorage.SetNumZeroed(AttributeListSize);
+	LPPROC_THREAD_ATTRIBUTE_LIST AttributeList = reinterpret_cast<LPPROC_THREAD_ATTRIBUTE_LIST>(AttributeListStorage.GetData());
+	if (!AttributeList)
+	{
+		UE_LOG(LogTerminal, Error, TEXT("Failed to allocate attribute list storage."));
+		if (API.ClosePseudoConsole)
+		{
+			API.ClosePseudoConsole(PseudoConsole);
+		}
+		CloseHandle(InputPipeReadHandle);
+		CloseHandle(InputPipeWriteHandle);
+		CloseHandle(OutputPipeReadHandle);
+		CloseHandle(OutputPipeWriteHandle);
+		return false;
+	}
+
+	if (!InitializeProcThreadAttributeList(AttributeList, 1, 0, &AttributeListSize))
+	{
+		UE_LOG(LogTerminal, Error, TEXT("InitializeProcThreadAttributeList failed."));
+		if (API.ClosePseudoConsole)
+		{
+			API.ClosePseudoConsole(PseudoConsole);
+		}
+		CloseHandle(InputPipeReadHandle);
+		CloseHandle(InputPipeWriteHandle);
+		CloseHandle(OutputPipeReadHandle);
+		CloseHandle(OutputPipeWriteHandle);
+		return false;
+	}
+
+	// PROC_THREAD_ATTRIBUTE_PSEUDOCONSOLE = 0x00020016
+	static const DWORD PROC_THREAD_ATTRIBUTE_PSEUDOCONSOLE_VALUE = 0x00020016;
+	if (!UpdateProcThreadAttribute(
+		AttributeList, 0,
+		PROC_THREAD_ATTRIBUTE_PSEUDOCONSOLE_VALUE,
+		PseudoConsole, sizeof(void*),
+		nullptr, nullptr))
+	{
+		UE_LOG(LogTerminal, Error, TEXT("UpdateProcThreadAttribute failed."));
+		DeleteProcThreadAttributeList(AttributeList);
+		if (API.ClosePseudoConsole)
+		{
+			API.ClosePseudoConsole(PseudoConsole);
+		}
+		CloseHandle(InputPipeReadHandle);
+		CloseHandle(InputPipeWriteHandle);
+		CloseHandle(OutputPipeReadHandle);
+		CloseHandle(OutputPipeWriteHandle);
+		return false;
+	}
+
+	STARTUPINFOEXW StartupInfoEx = {};
+	StartupInfoEx.StartupInfo.cb = sizeof(STARTUPINFOEXW);
+	StartupInfoEx.lpAttributeList = AttributeList;
+
+	PROCESS_INFORMATION ProcessInfo = {};
+
+	// CreateProcessW needs a mutable command line buffer. Quote the path in case it contains spaces.
+	const FString QuotedShellPath = FString::Printf(TEXT("\"%s\""), *EffectiveShellPath);
+	TArray<TCHAR> CommandLineBuffer;
+	CommandLineBuffer.Append(*QuotedShellPath, QuotedShellPath.Len() + 1);
+
+	LPCWSTR WorkingDir = WorkingDirectory.IsEmpty() ? nullptr : *WorkingDirectory;
+
+	// CreateProcessW copies the parent's STD handle values into the child's PEB
+	// regardless of bInheritHandles. When the editor is launched from a context
+	// that sets std handles to pipes or files (e.g., an IDE, build system, or
+	// script), the child shell inherits those handle values and uses them instead
+	// of the console handles provided by ConPTY. This causes the shell's output
+	// to bypass the pseudo-console entirely, resulting in a blank terminal.
+	// Clearing std handles to INVALID_HANDLE_VALUE before CreateProcessW ensures
+	// the child receives proper console handles from the pseudo-console.
+	HANDLE SavedStdInput = GetStdHandle(STD_INPUT_HANDLE);
+	HANDLE SavedStdOutput = GetStdHandle(STD_OUTPUT_HANDLE);
+	HANDLE SavedStdError = GetStdHandle(STD_ERROR_HANDLE);
+	if (!SetStdHandle(STD_INPUT_HANDLE, INVALID_HANDLE_VALUE)
+		|| !SetStdHandle(STD_OUTPUT_HANDLE, INVALID_HANDLE_VALUE)
+		|| !SetStdHandle(STD_ERROR_HANDLE, INVALID_HANDLE_VALUE))
+	{
+		UE_LOG(LogTerminal, Warning, TEXT("Failed to clear standard handles before CreateProcessW (error %u). Child may inherit parent handles."), GetLastError());
+	}
+
+	BOOL bSuccess = CreateProcessW(
+		*EffectiveShellPath,
+		CommandLineBuffer.GetData(),
+		nullptr,
+		nullptr,
+		0,
+		EXTENDED_STARTUPINFO_PRESENT,
+		nullptr,
+		WorkingDir,
+		&StartupInfoEx.StartupInfo,
+		&ProcessInfo);
+
+	if (!SetStdHandle(STD_INPUT_HANDLE, SavedStdInput)
+		|| !SetStdHandle(STD_OUTPUT_HANDLE, SavedStdOutput)
+		|| !SetStdHandle(STD_ERROR_HANDLE, SavedStdError))
+	{
+		UE_LOG(LogTerminal, Warning, TEXT("Failed to restore standard handles after CreateProcessW (error %u)."), GetLastError());
+	}
+
+	DeleteProcThreadAttributeList(AttributeList);
+
+	// Close the pipe ends that belong to the child process.
+	CloseHandle(InputPipeReadHandle);
+	CloseHandle(OutputPipeWriteHandle);
+
+	if (!bSuccess)
+	{
+		UE_LOG(LogTerminal, Error, TEXT("CreateProcessW failed for shell: %s (error %d)."), *EffectiveShellPath, GetLastError());
+		if (API.ClosePseudoConsole)
+		{
+			API.ClosePseudoConsole(PseudoConsole);
+		}
+		CloseHandle(InputPipeWriteHandle);
+		CloseHandle(OutputPipeReadHandle);
+		return false;
+	}
+
+	// All setup succeeded  - commit handles to member state.
+	PseudoConsoleHandle = PseudoConsole;
+	InputPipeWrite = InputPipeWriteHandle;
+	OutputPipeRead = OutputPipeReadHandle;
+	ProcessHandle = ProcessInfo.hProcess;
+	ThreadHandle = ProcessInfo.hThread;
+	bAlive = MakeShared<TAtomic<bool>>(true);
+
+	// Start the background reader thread.
+	ReaderRunnable = MakeUnique<FReaderRunnable>(this);
+	ReaderThread = TUniquePtr<FRunnableThread>(FRunnableThread::Create(
+		ReaderRunnable.Get(), TEXT("TerminalReader"), 0, TPri_Normal));
+
+	if (!ReaderThread)
+	{
+		UE_LOG(LogTerminal, Error, TEXT("Failed to create reader thread for shell: %s"), *EffectiveShellPath);
+		Shutdown();
+		return false;
+	}
+
+	UE_LOG(LogTerminal, Display, TEXT("Terminal session created with shell: %s"), *EffectiveShellPath);
+	return true;
+}
+
+void FConPTYSession::Shutdown()
+{
+	// Signal that this session is no longer alive so async callbacks are safe.
+	if (bAlive)
+	{
+		*bAlive = false;
+	}
+
+	// Stop the reader thread first.
+	if (ReaderRunnable.IsValid())
+	{
+		ReaderRunnable->Stop();
+	}
+
+	// Close the pseudo-console  - this will cause ReadFile on the output pipe to return.
+	if (PseudoConsoleHandle)
+	{
+		const FConPTYAPI& API = GetConPTYAPI();
+		if (API.ClosePseudoConsole)
+		{
+			API.ClosePseudoConsole(PseudoConsoleHandle);
+		}
+		PseudoConsoleHandle = nullptr;
+	}
+
+	// Wait for the reader thread to finish.
+	if (ReaderThread.IsValid())
+	{
+		ReaderThread->WaitForCompletion();
+		ReaderThread.Reset();
+	}
+	ReaderRunnable.Reset();
+
+	// Terminate the shell process if still running.
+	if (ProcessHandle)
+	{
+		DWORD ExitCode = 0;
+		if (GetExitCodeProcess(static_cast<HANDLE>(ProcessHandle), &ExitCode))
+		{
+			if (ExitCode == STILL_ACTIVE)
+			{
+				TerminateProcess(static_cast<HANDLE>(ProcessHandle), 0);
+			}
+		}
+		CloseHandle(static_cast<HANDLE>(ProcessHandle));
+		ProcessHandle = nullptr;
+	}
+
+	if (ThreadHandle)
+	{
+		CloseHandle(static_cast<HANDLE>(ThreadHandle));
+		ThreadHandle = nullptr;
+	}
+
+	if (InputPipeWrite)
+	{
+		CloseHandle(static_cast<HANDLE>(InputPipeWrite));
+		InputPipeWrite = nullptr;
+	}
+
+	if (OutputPipeRead)
+	{
+		CloseHandle(static_cast<HANDLE>(OutputPipeRead));
+		OutputPipeRead = nullptr;
+	}
+}
+
+void FConPTYSession::WriteInput(TArrayView<const uint8> Data)
+{
+	if (!InputPipeWrite || Data.Num() == 0)
+	{
+		return;
+	}
+
+	int32 Offset = 0;
+	while (Offset < Data.Num())
+	{
+		DWORD BytesWritten = 0;
+		if (!WriteFile(static_cast<HANDLE>(InputPipeWrite), Data.GetData() + Offset, Data.Num() - Offset, &BytesWritten, nullptr))
+		{
+			UE_LOG(LogTerminal, Warning, TEXT("WriteInput: WriteFile failed (error %d). Attempted %d bytes."), GetLastError(), Data.Num() - Offset);
+			break;
+		}
+		Offset += static_cast<int32>(BytesWritten);
+	}
+	UE_LOG(LogTerminal, VeryVerbose, TEXT("WriteInput: wrote %d/%d bytes to input pipe."), Offset, Data.Num());
+}
+
+void FConPTYSession::Resize(int32 Columns, int32 Rows)
+{
+	if (!PseudoConsoleHandle)
+	{
+		return;
+	}
+
+	const FConPTYAPI& API = GetConPTYAPI();
+	if (API.ResizePseudoConsole)
+	{
+		COORD NewSize;
+		NewSize.X = static_cast<SHORT>(FMath::Clamp(Columns, 1, static_cast<int32>(SHRT_MAX)));
+		NewSize.Y = static_cast<SHORT>(FMath::Clamp(Rows, 1, static_cast<int32>(SHRT_MAX)));
+		API.ResizePseudoConsole(PseudoConsoleHandle, NewSize);
+	}
+}
+
+TArray<uint8> FConPTYSession::ConsumeOutput()
+{
+	FScopeLock Lock(&StagingLock);
+	TArray<uint8> Output = MoveTemp(StagingBuffer);
+	StagingBuffer.Reset();
+	return Output;
+}
+
+bool FConPTYSession::IsRunning() const
+{
+	if (!ProcessHandle)
+	{
+		return false;
+	}
+
+	DWORD ExitCode = 0;
+	if (GetExitCodeProcess(static_cast<HANDLE>(ProcessHandle), &ExitCode))
+	{
+		return ExitCode == STILL_ACTIVE;
+	}
+	return false;
+}
+
+uint32 FConPTYSession::FReaderRunnable::Run()
+{
+	// Cache the pipe handle locally so we do not re-read through the Owner pointer
+	// on every iteration. Shutdown() guarantees OutputPipeRead stays valid until
+	// after WaitForCompletion(), but caching avoids a TOCTOU gap between the null
+	// check and the ReadFile call.
+	const HANDLE CachedOutputPipeRead = static_cast<HANDLE>(Owner->OutputPipeRead);
+	if (!CachedOutputPipeRead || CachedOutputPipeRead == INVALID_HANDLE_VALUE)
+	{
+		return 0;
+	}
+
+	UE_LOG(LogTerminal, VeryVerbose, TEXT("Reader thread started. Pipe handle: %p."), CachedOutputPipeRead);
+
+	constexpr int32 ReadBufferSize = 65536;
+	constexpr int32 MaxStagingBufferSize = 16 * 1024 * 1024; // 16 MB backpressure cap.
+	TArray<uint8> ReadBuffer;
+	ReadBuffer.SetNumUninitialized(ReadBufferSize);
+	int64 TotalBytesRead = 0;
+
+	while (!bStopRequested)
+	{
+		DWORD BytesRead = 0;
+		BOOL bSuccess = ReadFile(
+			CachedOutputPipeRead,
+			ReadBuffer.GetData(),
+			ReadBufferSize,
+			&BytesRead,
+			nullptr);
+
+		if (!bSuccess || BytesRead == 0)
+		{
+			const DWORD LastError = bSuccess ? 0 : GetLastError();
+			UE_LOG(LogTerminal, VeryVerbose, TEXT("Reader thread: ReadFile returned success=%d bytes=%d error=%d. Total bytes read: %lld."), bSuccess, BytesRead, LastError, TotalBytesRead);
+			break;
+		}
+
+		TotalBytesRead += BytesRead;
+		UE_LOG(LogTerminal, VeryVerbose, TEXT("Reader thread: ReadFile returned %d bytes (total: %lld)."), BytesRead, TotalBytesRead);
+
+		{
+			FScopeLock Lock(&Owner->StagingLock);
+			const int32 Remaining = MaxStagingBufferSize - Owner->StagingBuffer.Num();
+			if (Remaining > 0)
+			{
+				const int32 BytesToAppend = FMath::Min(static_cast<int32>(BytesRead), Remaining);
+				Owner->StagingBuffer.Append(ReadBuffer.GetData(), BytesToAppend);
+				if (BytesToAppend < static_cast<int32>(BytesRead))
+				{
+					UE_LOG(LogTerminal, Warning, TEXT("Staging buffer backpressure: %d of %d bytes dropped (cap %d)."),
+						BytesRead - BytesToAppend, BytesRead, MaxStagingBufferSize);
+				}
+			}
+			else
+			{
+				UE_LOG(LogTerminal, Warning, TEXT("Staging buffer full: %d bytes dropped (cap %d)."), BytesRead, MaxStagingBufferSize);
+			}
+		}
+	}
+
+	UE_LOG(LogTerminal, VeryVerbose, TEXT("Reader thread exiting. Total bytes read: %lld. StopRequested: %d."), TotalBytesRead, bStopRequested ? 1 : 0);
+
+	// Check for process exit.
+	if (Owner->ProcessHandle)
+	{
+		DWORD ExitCode = 0;
+		if (GetExitCodeProcess(static_cast<HANDLE>(Owner->ProcessHandle), &ExitCode) && ExitCode != STILL_ACTIVE)
+		{
+			AsyncTask(ENamedThreads::GameThread, [Alive = Owner->bAlive, Delegate = Owner->OnProcessExited, ExitCode]()
+			{
+				if (Alive && *Alive)
+				{
+					Delegate.ExecuteIfBound(static_cast<int32>(ExitCode));
+				}
+			});
+		}
+	}
+
+	return 0;
+}
+
+void FConPTYSession::FReaderRunnable::Stop()
+{
+	bStopRequested = true;
+}
+
+#endif // PLATFORM_WINDOWS
diff --git a/Engine/Plugins/Experimental/Terminal/Source/Terminal/Private/ITerminalSession.cpp b/Engine/Plugins/Experimental/Terminal/Source/Terminal/Private/ITerminalSession.cpp
new file mode 100644
--- /dev/null
+++ b/Engine/Plugins/Experimental/Terminal/Source/Terminal/Private/ITerminalSession.cpp
@@ -0,0 +1,26 @@
+// Copyright Epic Games, Inc. All Rights Reserved.
+
+#include "ITerminalSession.h"
+
+#if PLATFORM_WINDOWS
+#include "ConPTYSession.h"
+#elif PLATFORM_UNIX || PLATFORM_MAC
+#include "PosixPTYSession.h"
+#endif
+
+TSharedPtr<ITerminalSession> ITerminalSession::CreateForCurrentPlatform(FString& OutError)
+{
+#if PLATFORM_WINDOWS
+	if (FConPTYSession::IsConPTYAvailable())
+	{
+		return MakeShared<FConPTYSession>();
+	}
+	OutError = TEXT("ConPTY API is not available on this system.");
+	return nullptr;
+#elif PLATFORM_UNIX || PLATFORM_MAC
+	return MakeShared<FPosixPTYSession>();
+#else
+	OutError = TEXT("Terminal is not supported on this platform.");
+	return nullptr;
+#endif
+}
diff --git a/Engine/Plugins/Experimental/Terminal/Source/Terminal/Private/PosixPTYSession.cpp b/Engine/Plugins/Experimental/Terminal/Source/Terminal/Private/PosixPTYSession.cpp
new file mode 100644
--- /dev/null
+++ b/Engine/Plugins/Experimental/Terminal/Source/Terminal/Private/PosixPTYSession.cpp
@@ -0,0 +1,330 @@
+// Copyright Epic Games, Inc. All Rights Reserved.
+
+#include "PosixPTYSession.h"
+
+#if PLATFORM_UNIX || PLATFORM_MAC
+
+#include "Async/Async.h"
+#include "Containers/Utf8String.h"
+#include "Misc/Paths.h"
+
+#include <unistd.h>
+#include <fcntl.h>
+#include <signal.h>
+#include <sys/ioctl.h>
+#include <sys/wait.h>
+#include <termios.h>
+
+#if PLATFORM_MAC
+// Forward-declare forkpty to avoid including <util.h>, which transitively
+// pulls in <pwd.h>  - absent from some CI Apple SDK configurations.
+extern "C" int forkpty(int*, char*, struct termios*, struct winsize*);
+#else
+#include <pty.h>
+#endif
+
+FPosixPTYSession::FPosixPTYSession()
+{
+}
+
+FPosixPTYSession::~FPosixPTYSession()
+{
+	Shutdown();
+}
+
+bool FPosixPTYSession::Create(const FString& ShellPath, const FString& WorkingDirectory, int32 Columns, int32 Rows)
+{
+	// Set up the initial terminal window size.
+	struct winsize WindowSize = {};
+	WindowSize.ws_col = static_cast<unsigned short>(FMath::Clamp(Columns, 1, 32767));
+	WindowSize.ws_row = static_cast<unsigned short>(FMath::Clamp(Rows, 1, 32767));
+
+	int MasterFD = -1;
+	pid_t ChildPid = forkpty(&MasterFD, nullptr, nullptr, &WindowSize);
+
+	if (ChildPid < 0)
+	{
+		UE_LOG(LogTerminal, Error, TEXT("forkpty() failed with errno %d."), errno);
+		return false;
+	}
+
+	if (ChildPid == 0)
+	{
+		// Child process  - exec the shell.
+		FString EffectiveShellPath = ShellPath;
+		if (EffectiveShellPath.IsEmpty())
+		{
+			const char* ShellEnv = getenv("SHELL");
+			if (ShellEnv && ShellEnv[0] != '\0')
+			{
+				EffectiveShellPath = UTF8_TO_TCHAR(ShellEnv);
+			}
+			else
+			{
+				EffectiveShellPath = TEXT("/bin/bash");
+			}
+		}
+
+		if (!WorkingDirectory.IsEmpty())
+		{
+			const FUtf8String DirectoryUtf8(WorkingDirectory);
+			if (chdir(reinterpret_cast<const char*>(*DirectoryUtf8)) != 0)
+			{
+				// Cannot log from the forked child; proceed in the default directory.
+				// The parent will detect issues through normal shell behavior.
+			}
+		}
+
+		// Set TERM so curses/readline applications behave correctly.
+		setenv("TERM", "xterm-256color", 1);
+
+		const FUtf8String ShellUtf8(EffectiveShellPath);
+		const char* ShellCStr = reinterpret_cast<const char*>(*ShellUtf8);
+
+		// Extract just the shell name for argv[0] (convention: login shells use "-bash" etc.).
+		const char* ShellName = strrchr(ShellCStr, '/');
+		ShellName = ShellName ? ShellName + 1 : ShellCStr;
+
+		// execvp requires a mutable argv. Use a login-shell style argv[0].
+		char LoginShellName[256];
+		snprintf(LoginShellName, sizeof(LoginShellName), "-%s", ShellName);
+
+		char* Argv[] = { LoginShellName, nullptr };
+		execvp(ShellCStr, Argv);
+
+		// exec failed  - exit the forked child.
+		_exit(127);
+	}
+
+	// Parent process.
+	MasterFileDescriptor = MasterFD;
+	ProcessId = ChildPid;
+	bAlive = MakeShared<TAtomic<bool>>(true);
+
+	UE_LOG(LogTerminal, VeryVerbose, TEXT("Resolved shell path: %s (requested: %s). Terminal size: %dx%d. PID: %d."),
+		*ShellPath, ShellPath.IsEmpty() ? TEXT("<default>") : *ShellPath, Columns, Rows, static_cast<int32>(ProcessId));
+
+	// Start the background reader thread.
+	ReaderRunnable = MakeUnique<FReaderRunnable>(this);
+	ReaderThread = TUniquePtr<FRunnableThread>(FRunnableThread::Create(
+		ReaderRunnable.Get(), TEXT("TerminalReader"), 0, TPri_Normal));
+
+	if (!ReaderThread)
+	{
+		UE_LOG(LogTerminal, Error, TEXT("Failed to create reader thread for shell PID %d."), static_cast<int32>(ProcessId));
+		Shutdown();
+		return false;
+	}
+
+	UE_LOG(LogTerminal, Display, TEXT("Terminal session created with shell PID %d."), static_cast<int32>(ProcessId));
+	return true;
+}
+
+void FPosixPTYSession::Shutdown()
+{
+	// Signal that this session is no longer alive so async callbacks are safe.
+	if (bAlive)
+	{
+		*bAlive = false;
+	}
+
+	// Stop the reader thread first.
+	if (ReaderRunnable.IsValid())
+	{
+		ReaderRunnable->Stop();
+	}
+
+	// Close the master FD  - this signals EOF to the reader thread's read() call.
+	if (MasterFileDescriptor >= 0)
+	{
+		close(MasterFileDescriptor);
+		MasterFileDescriptor = -1;
+	}
+
+	// Wait for the reader thread to finish.
+	if (ReaderThread.IsValid())
+	{
+		ReaderThread->WaitForCompletion();
+		ReaderThread.Reset();
+	}
+	ReaderRunnable.Reset();
+
+	// Reap the child process.
+	if (ProcessId > 0)
+	{
+		int Status = 0;
+		pid_t Result = waitpid(ProcessId, &Status, WNOHANG);
+		if (Result < 0)
+		{
+			UE_LOG(LogTerminal, Warning, TEXT("Shutdown: waitpid(%d) failed (errno %d)."), static_cast<int32>(ProcessId), errno);
+		}
+		else if (Result == 0)
+		{
+			// Child still running  - send SIGHUP then SIGKILL.
+			if (kill(ProcessId, SIGHUP) != 0)
+			{
+				UE_LOG(LogTerminal, Warning, TEXT("Shutdown: kill(%d, SIGHUP) failed (errno %d)."), static_cast<int32>(ProcessId), errno);
+			}
+			usleep(100000); // 100ms grace period.
+			Result = waitpid(ProcessId, &Status, WNOHANG);
+			if (Result == 0)
+			{
+				if (kill(ProcessId, SIGKILL) != 0)
+				{
+					UE_LOG(LogTerminal, Warning, TEXT("Shutdown: kill(%d, SIGKILL) failed (errno %d)."), static_cast<int32>(ProcessId), errno);
+				}
+				waitpid(ProcessId, &Status, 0);
+			}
+			else if (Result < 0)
+			{
+				UE_LOG(LogTerminal, Warning, TEXT("Shutdown: waitpid(%d) after SIGHUP failed (errno %d)."), static_cast<int32>(ProcessId), errno);
+			}
+		}
+		ProcessId = -1;
+	}
+}
+
+void FPosixPTYSession::WriteInput(TArrayView<const uint8> Data)
+{
+	if (MasterFileDescriptor < 0 || Data.Num() == 0)
+	{
+		return;
+	}
+
+	int32 Offset = 0;
+	while (Offset < Data.Num())
+	{
+		ssize_t BytesWritten = write(MasterFileDescriptor, Data.GetData() + Offset, Data.Num() - Offset);
+		if (BytesWritten <= 0)
+		{
+			if (BytesWritten < 0 && errno == EINTR)
+			{
+				continue;
+			}
+			UE_LOG(LogTerminal, Warning, TEXT("WriteInput: write() returned %zd (errno %d). Attempted %d bytes."), BytesWritten, BytesWritten < 0 ? errno : 0, Data.Num() - Offset);
+			break;
+		}
+		Offset += static_cast<int32>(BytesWritten);
+	}
+	UE_LOG(LogTerminal, VeryVerbose, TEXT("WriteInput: wrote %d/%d bytes to master FD."), Offset, Data.Num());
+}
+
+void FPosixPTYSession::Resize(int32 Columns, int32 Rows)
+{
+	if (MasterFileDescriptor < 0)
+	{
+		return;
+	}
+
+	struct winsize WindowSize = {};
+	WindowSize.ws_col = static_cast<unsigned short>(FMath::Clamp(Columns, 1, 32767));
+	WindowSize.ws_row = static_cast<unsigned short>(FMath::Clamp(Rows, 1, 32767));
+	if (ioctl(MasterFileDescriptor, TIOCSWINSZ, &WindowSize) != 0)
+	{
+		UE_LOG(LogTerminal, Warning, TEXT("Resize: ioctl(TIOCSWINSZ) failed (errno %d) for %dx%d."), errno, Columns, Rows);
+	}
+}
+
+TArray<uint8> FPosixPTYSession::ConsumeOutput()
+{
+	FScopeLock Lock(&StagingLock);
+	TArray<uint8> Output = MoveTemp(StagingBuffer);
+	StagingBuffer.Reset();
+	return Output;
+}
+
+bool FPosixPTYSession::IsRunning() const
+{
+	if (ProcessId <= 0)
+	{
+		return false;
+	}
+
+	// Use kill with signal 0 to check liveness without reaping the child.
+	// waitpid would reap the process, racing with Shutdown() and the reader thread.
+	return kill(ProcessId, 0) == 0;
+}
+
+uint32 FPosixPTYSession::FReaderRunnable::Run()
+{
+	const int CachedMasterFD = Owner->MasterFileDescriptor;
+	if (CachedMasterFD < 0)
+	{
+		return 0;
+	}
+
+	UE_LOG(LogTerminal, VeryVerbose, TEXT("Reader thread started. Master FD: %d."), CachedMasterFD);
+
+	constexpr int32 ReadBufferSize = 65536;
+	constexpr int32 MaxStagingBufferSize = 16 * 1024 * 1024; // 16 MB backpressure cap.
+	TArray<uint8> ReadBuffer;
+	ReadBuffer.SetNumUninitialized(ReadBufferSize);
+	int64 TotalBytesRead = 0;
+
+	while (!bStopRequested)
+	{
+		ssize_t BytesRead = read(CachedMasterFD, ReadBuffer.GetData(), ReadBufferSize);
+
+		if (BytesRead <= 0)
+		{
+			if (BytesRead < 0 && errno == EINTR)
+			{
+				continue;
+			}
+			UE_LOG(LogTerminal, VeryVerbose, TEXT("Reader thread: read() returned %zd (errno %d). Total bytes read: %lld."),
+				BytesRead, BytesRead < 0 ? errno : 0, TotalBytesRead);
+			break;
+		}
+
+		TotalBytesRead += BytesRead;
+		UE_LOG(LogTerminal, VeryVerbose, TEXT("Reader thread: read() returned %zd bytes (total: %lld)."), BytesRead, TotalBytesRead);
+
+		{
+			FScopeLock Lock(&Owner->StagingLock);
+			const int32 Remaining = MaxStagingBufferSize - Owner->StagingBuffer.Num();
+			if (Remaining > 0)
+			{
+				const int32 BytesToAppend = FMath::Min(static_cast<int32>(BytesRead), Remaining);
+				Owner->StagingBuffer.Append(ReadBuffer.GetData(), BytesToAppend);
+				if (BytesToAppend < static_cast<int32>(BytesRead))
+				{
+					UE_LOG(LogTerminal, Warning, TEXT("Staging buffer backpressure: %zd of %zd bytes dropped (cap %d)."),
+						BytesRead - BytesToAppend, BytesRead, MaxStagingBufferSize);
+				}
+			}
+			else
+			{
+				UE_LOG(LogTerminal, Warning, TEXT("Staging buffer full: %zd bytes dropped (cap %d)."), BytesRead, MaxStagingBufferSize);
+			}
+		}
+	}
+
+	UE_LOG(LogTerminal, VeryVerbose, TEXT("Reader thread exiting. Total bytes read: %lld. StopRequested: %d."), TotalBytesRead, bStopRequested ? 1 : 0);
+
+	// Check for process exit.
+	if (Owner->ProcessId > 0)
+	{
+		int Status = 0;
+		pid_t Result = waitpid(Owner->ProcessId, &Status, WNOHANG);
+		if (Result > 0)
+		{
+			int32 ExitCode = WIFEXITED(Status) ? WEXITSTATUS(Status) : -1;
+			AsyncTask(ENamedThreads::GameThread, [Alive = Owner->bAlive, Delegate = Owner->OnProcessExited, ExitCode]()
+			{
+				if (Alive && *Alive)
+				{
+					Delegate.ExecuteIfBound(ExitCode);
+				}
+			});
+		}
+	}
+
+	return 0;
+}
+
+void FPosixPTYSession::FReaderRunnable::Stop()
+{
+	bStopRequested = true;
+}
+
+#endif // PLATFORM_UNIX || PLATFORM_MAC
diff --git a/Engine/Plugins/Experimental/Terminal/Source/Terminal/Private/STerminal.cpp b/Engine/Plugins/Experimental/Terminal/Source/Terminal/Private/STerminal.cpp
new file mode 100644
--- /dev/null
+++ b/Engine/Plugins/Experimental/Terminal/Source/Terminal/Private/STerminal.cpp
@@ -0,0 +1,1512 @@
+// Copyright Epic Games, Inc. All Rights Reserved.
+
+#include "STerminal.h"
+
+#include "ITerminalSession.h"
+#include "TerminalSettings.h"
+#include "TerminalSubsystem.h"
+#include "TerminalUtilities.h"
+#include "Editor.h"
+
+#include "Fonts/FontCache.h"
+#include "Framework/Application/SlateApplication.h"
+#include "HAL/PlatformApplicationMisc.h"
+#include "ProfilingDebugging/CpuProfilerTrace.h"
+#include "Rendering/DrawElements.h"
+#include "Widgets/Layout/SScrollBar.h"
+
+static TAutoConsoleVariable<bool> CVarTerminalDebugGlyphs(
+	TEXT("Terminal.DebugGlyphs"),
+	false,
+	TEXT("When true, draws per-cell glyph boundary boxes to diagnose font measurement vs rendering mismatches."),
+	ECVF_Default);
+
+void STerminal::Construct(const FArguments& InArgs)
+{
+	// Set up the font. Buffer and session creation are deferred to the first paint
+	// so we know the actual widget dimensions and avoid a spurious resize.
+	const UTerminalSettings* Settings = GetDefault<UTerminalSettings>();
+	const FString FontFamily = Settings ? Settings->FontFamily : TEXT("Consolas");
+	const int32 FontSize = Settings ? Settings->FontSize : 10;
+
+	// Resolve font: try configured font, then common monospace fallbacks, then engine font.
+	FString FontFilePath = TerminalUtilities::ResolveSystemFontPath(FontFamily);
+	if (FontFilePath.IsEmpty())
+	{
+		for (const TCHAR* Fallback : TerminalUtilities::GetFallbackFontNames())
+		{
+			if (FontFamily != Fallback)
+			{
+				FontFilePath = TerminalUtilities::ResolveSystemFontPath(Fallback);
+				if (!FontFilePath.IsEmpty())
+				{
+					break;
+				}
+			}
+		}
+	}
+	
+	if (FontFilePath.IsEmpty())
+	{
+		UE_LOG(LogTerminal, Warning, TEXT("No system monospace font found. Falling back to engine default."));
+		FontFilePath = TerminalUtilities::GetEngineFallbackFontPath();
+	}
+	
+	UE_LOG(LogTerminal, Display, TEXT("Using terminal font: %s"), *FontFilePath);
+
+	TerminalFont = MakeShared<FStandaloneCompositeFont>(NAME_None, FontFilePath, EFontHinting::Default, EFontLoadingPolicy::LazyLoad);
+
+	// Add symbol/emoji fallback sub-typefaces so glyphs not in the monospace
+	// font (box drawing, arrows, math operators, emoji, etc.) still render.
+	{
+		for (const TCHAR* SymbolFontName : TerminalUtilities::GetSymbolFontNames())
+		{
+			const FString SymbolFontPath = TerminalUtilities::ResolveSystemFontPath(SymbolFontName);
+			if (!SymbolFontPath.IsEmpty())
+			{
+				FCompositeSubFont& SymbolSubFont = TerminalFont->SubTypefaces.AddDefaulted_GetRef();
+				// Broad range covering symbols, box drawing, block elements, geometric shapes,
+				// arrows, mathematical operators, Braille, dingbats, and miscellaneous symbols.
+				SymbolSubFont.CharacterRanges.Add(FInt32Range(FInt32Range::BoundsType::Inclusive(0x2000), FInt32Range::BoundsType::Inclusive(0x2BFF)));
+				// Supplemental arrows and math
+				SymbolSubFont.CharacterRanges.Add(FInt32Range(FInt32Range::BoundsType::Inclusive(0x2E80), FInt32Range::BoundsType::Inclusive(0x2FFF)));
+				SymbolSubFont.Typeface.AppendFont(TEXT("Regular"), SymbolFontPath, EFontHinting::Default, EFontLoadingPolicy::LazyLoad);
+				break;
+			}
+		}
+
+		for (const TCHAR* EmojiFontName : TerminalUtilities::GetEmojiFontNames())
+		{
+			const FString EmojiFontPath = TerminalUtilities::ResolveSystemFontPath(EmojiFontName);
+			if (!EmojiFontPath.IsEmpty())
+			{
+				FCompositeSubFont& EmojiSubFont = TerminalFont->SubTypefaces.AddDefaulted_GetRef();
+				// Supplemental symbols and pictographs (0x2600-0x27BF already covered by Symbol sub-typeface).
+				EmojiSubFont.CharacterRanges.Add(FInt32Range(FInt32Range::BoundsType::Inclusive(0x1F300), FInt32Range::BoundsType::Inclusive(0x1FAFF)));
+				EmojiSubFont.Typeface.AppendFont(TEXT("Regular"), EmojiFontPath, EFontHinting::Default, EFontLoadingPolicy::LazyLoad);
+				break;
+			}
+		}
+	}
+
+	FontInfo = FSlateFontInfo(TerminalFont, static_cast<float>(FontSize));
+	FontInfo.FontFallback = EFontFallback::FF_Max;
+
+	MeasureCellSize();
+
+	ScrollBar = InArgs._ExternalScrollbar;
+	if (ScrollBar.IsValid())
+	{
+		ScrollBar->SetOnUserScrolled(FOnUserScrolled::CreateSP(this, &STerminal::HandleScrollBarScrolled));
+		ScrollBar->SetState(0.0f, 1.0f);
+	}
+
+	UE_LOG(LogTerminal, VeryVerbose, TEXT("STerminal::Construct complete. CellSize=%.1fx%.1f"), CellWidth, CellHeight);
+
+	// Register tick for consuming ConPTY output.
+	RegisterActiveTimer(0.016f, FWidgetActiveTimerDelegate::CreateSP(this, &STerminal::OnTick));
+
+	// Register cursor blink timer (~530ms).
+	RegisterActiveTimer(0.530f, FWidgetActiveTimerDelegate::CreateSP(this, &STerminal::OnCursorBlink));
+}
+
+void STerminal::ExecuteCommand(const FString& Command)
+{
+	if (!Session.IsValid())
+	{
+		return;
+	}
+
+	const FUtf8String CommandUtf8(Command + TEXT("\r"));
+	Session->WriteInput(MakeArrayView(reinterpret_cast<const uint8*>(*CommandUtf8), CommandUtf8.Len()));
+}
+
+bool STerminal::IsSessionRunning() const
+{
+	return Session.IsValid() && Session->IsRunning();
+}
+
+STerminal::~STerminal()
+{
+	UE_LOG(LogTerminal, VeryVerbose, TEXT("~STerminal: total bytes consumed by parser: %lld. Cursor at (%d,%d). Initialized: %d."), TotalBytesConsumed, Buffer.Cursor.Column, Buffer.Cursor.Row, bInitialized ? 1 : 0);
+
+	if (Session.IsValid())
+	{
+		Session->Shutdown();
+		Session.Reset();
+	}
+}
+
+FVector2D STerminal::ComputeDesiredSize(float LayoutScaleMultiplier) const
+{
+	return FVector2D(ViewportColumns * CellWidth, ViewportRows * CellHeight);
+}
+
+int32 STerminal::OnPaint(const FPaintArgs& Args, const FGeometry& AllottedGeometry, const FSlateRect& MyCullingRect, FSlateWindowElementList& OutDrawElements, int32 LayerId, const FWidgetStyle& InWidgetStyle, bool bParentEnabled) const
+{
+	TRACE_CPUPROFILER_EVENT_SCOPE(STerminal::OnPaint);
+
+	// Detect size changes for resize handling. Slate's OnPaint is const, but deferred
+	// initialization and viewport resize must mutate state on the first (or resized) paint.
+	// This is a standard Slate pattern for geometry-dependent lazy init.
+	// When not yet initialized, always attempt UpdateViewportDimensions so that
+	// initialization is not skipped when the first paint arrives with zero geometry
+	// (which would match the zero-initialized LastAllocatedSize and be a no-op).
+	const FVector2D CurrentSize = AllottedGeometry.GetLocalSize();
+	if (!bInitialized || !CurrentSize.Equals(LastAllocatedSize, 1.0f))
+	{
+		LastAllocatedSize = CurrentSize;
+		const_cast<STerminal*>(this)->UpdateViewportDimensions(AllottedGeometry);
+	}
+
+	// Draw full background.
+	FSlateDrawElement::MakeBox(
+		OutDrawElements,
+		LayerId,
+		AllottedGeometry.ToPaintGeometry(),
+		FAppStyle::GetBrush("WhiteBrush"),
+		ESlateDrawEffect::None,
+		FLinearColor(Buffer.DefaultCell.Background));
+
+	if (bInitialized)
+	{
+		PaintBackground(AllottedGeometry, OutDrawElements, ++LayerId);
+		PaintText(AllottedGeometry, OutDrawElements, ++LayerId);
+		PaintForeground(AllottedGeometry, OutDrawElements, ++LayerId);
+	}
+
+	return LayerId;
+}
+
+FColor STerminal::ComputeForegroundColor(const FTerminalCell& Cell)
+{
+	FColor Color = Cell.Foreground;
+	if (Cell.Attributes & ETerminalAttribute::Inverse)
+	{
+		Color = Cell.Background;
+	}
+
+	if (Cell.Attributes & ETerminalAttribute::Bold)
+	{
+		Color.R = static_cast<uint8>(FMath::Min(255, Color.R + 60));
+		Color.G = static_cast<uint8>(FMath::Min(255, Color.G + 60));
+		Color.B = static_cast<uint8>(FMath::Min(255, Color.B + 60));
+	}
+	else if (Cell.Attributes & ETerminalAttribute::Dim)
+	{
+		Color.R /= 2;
+		Color.G /= 2;
+		Color.B /= 2;
+	}
+
+	return Color;
+}
+
+void STerminal::PaintBackground(const FGeometry& AllottedGeometry, FSlateWindowElementList& OutDrawElements, int32 LayerId) const
+{
+	TRACE_CPUPROFILER_EVENT_SCOPE(STerminal::PaintBackground);
+
+	const FColor DefaultBackground = Buffer.DefaultCell.Background;
+	const int32 TotalRows = Buffer.GetTotalRows();
+	const int32 Columns = Buffer.GetColumns();
+
+	for (int32 ViewportRow = 0; ViewportRow < ViewportRows; ++ViewportRow)
+	{
+		const int32 AbsoluteRow = Buffer.GetAbsoluteRow(ViewportRow, ScrollOffset);
+		if (AbsoluteRow < 0 || AbsoluteRow >= TotalRows)
+		{
+			continue;
+		}
+
+		const float Y = ViewportRow * CellHeight;
+		int32 Col = 0;
+		while (Col < Columns)
+		{
+			const FTerminalCell& StartCell = Buffer.GetCell(AbsoluteRow, Col);
+
+			FColor BackgroundColor = StartCell.Background;
+			if (StartCell.Attributes & ETerminalAttribute::Inverse)
+			{
+				BackgroundColor = StartCell.Foreground;
+			}
+
+			// Batch consecutive cells with same background.
+			int32 RunEnd = Col + 1;
+			while (RunEnd < Columns)
+			{
+				const FTerminalCell& NextCell = Buffer.GetCell(AbsoluteRow, RunEnd);
+
+				FColor NextBackground = NextCell.Background;
+				if (NextCell.Attributes & ETerminalAttribute::Inverse)
+				{
+					NextBackground = NextCell.Foreground;
+				}
+
+				if (NextBackground != BackgroundColor)
+				{
+					break;
+				}
+				RunEnd++;
+			}
+
+			if (BackgroundColor != DefaultBackground)
+			{
+				const float X = Col * CellWidth;
+				const float RunWidth = (RunEnd - Col) * CellWidth;
+				FSlateDrawElement::MakeBox(
+					OutDrawElements,
+					LayerId,
+					AllottedGeometry.ToPaintGeometry(FVector2D(RunWidth, CellHeight), FSlateLayoutTransform(FVector2D(X, Y))),
+					FAppStyle::GetBrush("WhiteBrush"),
+					ESlateDrawEffect::None,
+					FLinearColor(BackgroundColor));
+			}
+
+			Col = RunEnd;
+		}
+	}
+}
+
+void STerminal::PaintText(const FGeometry& AllottedGeometry, FSlateWindowElementList& OutDrawElements, int32 LayerId) const
+{
+	TRACE_CPUPROFILER_EVENT_SCOPE(STerminal::PaintText);
+
+	const int32 TotalRows = Buffer.GetTotalRows();
+	const int32 Columns = Buffer.GetColumns();
+
+	for (int32 ViewportRow = 0; ViewportRow < ViewportRows; ++ViewportRow)
+	{
+		const int32 AbsoluteRow = Buffer.GetAbsoluteRow(ViewportRow, ScrollOffset);
+		if (AbsoluteRow < 0 || AbsoluteRow >= TotalRows)
+		{
+			continue;
+		}
+
+		const float Y = ViewportRow * CellHeight;
+
+		// Batch runs of monospace-safe characters with the same foreground/attributes.
+		// Non-ASCII characters are rendered individually at their cell position to
+		// maintain grid alignment.
+		int32 Col = 0;
+		while (Col < Columns)
+		{
+			const FTerminalCell& StartCell = Buffer.GetCell(AbsoluteRow, Col);
+			const FColor ForegroundColor = ComputeForegroundColor(StartCell);
+			const bool bStartMonoSafe = TerminalUtilities::IsMonospaceSafe(StartCell.Character);
+
+			if (bStartMonoSafe)
+			{
+				// Batch consecutive monospace-safe characters with same foreground and attributes.
+				FString BatchText;
+				BatchText.Reserve(Columns - Col);
+				BatchText.AppendChar(StartCell.Character);
+				int32 RunEnd = Col + 1;
+				while (RunEnd < Columns)
+				{
+					const FTerminalCell& NextCell = Buffer.GetCell(AbsoluteRow, RunEnd);
+					if (!TerminalUtilities::IsMonospaceSafe(NextCell.Character))
+					{
+						break;
+					}
+
+					const FColor NextForeground = ComputeForegroundColor(NextCell);
+					if (NextForeground != ForegroundColor || NextCell.Attributes != StartCell.Attributes)
+					{
+						break;
+					}
+					BatchText.AppendChar(NextCell.Character);
+					RunEnd++;
+				}
+
+				const float X = Col * CellWidth;
+				const float RunWidth = (RunEnd - Col) * CellWidth;
+				const FLinearColor DrawColor(ForegroundColor);
+				FSlateDrawElement::MakeText(
+					OutDrawElements,
+					LayerId,
+					AllottedGeometry.ToPaintGeometry(FVector2D(RunWidth, CellHeight), FSlateLayoutTransform(FVector2D(X, Y))),
+					BatchText,
+					FontInfo,
+					ESlateDrawEffect::None,
+					DrawColor);
+
+				// Underline decoration for the entire run.
+				if (StartCell.Attributes & ETerminalAttribute::Underline)
+				{
+					FSlateDrawElement::MakeBox(
+						OutDrawElements,
+						LayerId,
+						AllottedGeometry.ToPaintGeometry(FVector2D(RunWidth, 1.0f), FSlateLayoutTransform(FVector2D(X, Y + CellHeight - 1.0f))),
+						FAppStyle::GetBrush("WhiteBrush"),
+						ESlateDrawEffect::None,
+						DrawColor);
+				}
+
+				// Strikethrough decoration for the entire run.
+				if (StartCell.Attributes & ETerminalAttribute::Strikethrough)
+				{
+					FSlateDrawElement::MakeBox(
+						OutDrawElements,
+						LayerId,
+						AllottedGeometry.ToPaintGeometry(FVector2D(RunWidth, 1.0f), FSlateLayoutTransform(FVector2D(X, Y + CellHeight * 0.5f))),
+						FAppStyle::GetBrush("WhiteBrush"),
+						ESlateDrawEffect::None,
+						DrawColor);
+				}
+
+				Col = RunEnd;
+			}
+			else
+			{
+				// Non-monospace-safe character: render at cell position without
+				// per-cell clipping. Fallback font glyphs may have advances
+				// slightly different from the monospace cell width; clipping to
+				// exact cell boundaries causes cut edges and anti-aliasing
+				// shadow artifacts. Any minor glyph overflow is covered by
+				// adjacent cells' backgrounds drawn on the layer below.
+				const float CharX = Col * CellWidth;
+				const FString CharString(1, &StartCell.Character);
+				const FLinearColor DrawColor(ForegroundColor);
+
+				FSlateDrawElement::MakeText(
+					OutDrawElements,
+					LayerId,
+					AllottedGeometry.ToPaintGeometry(
+						FVector2D(CellWidth, CellHeight), FSlateLayoutTransform(FVector2D(CharX, Y))),
+					CharString,
+					FontInfo,
+					ESlateDrawEffect::None,
+					DrawColor);
+
+				// Underline.
+				if (StartCell.Attributes & ETerminalAttribute::Underline)
+				{
+					FSlateDrawElement::MakeBox(
+						OutDrawElements,
+						LayerId,
+						AllottedGeometry.ToPaintGeometry(FVector2D(CellWidth, 1.0f), FSlateLayoutTransform(FVector2D(CharX, Y + CellHeight - 1.0f))),
+						FAppStyle::GetBrush("WhiteBrush"),
+						ESlateDrawEffect::None,
+						DrawColor);
+				}
+
+				// Strikethrough.
+				if (StartCell.Attributes & ETerminalAttribute::Strikethrough)
+				{
+					FSlateDrawElement::MakeBox(
+						OutDrawElements,
+						LayerId,
+						AllottedGeometry.ToPaintGeometry(FVector2D(CellWidth, 1.0f), FSlateLayoutTransform(FVector2D(CharX, Y + CellHeight * 0.5f))),
+						FAppStyle::GetBrush("WhiteBrush"),
+						ESlateDrawEffect::None,
+						DrawColor);
+				}
+
+				Col++;
+			}
+		}
+	}
+}
+
+void STerminal::PaintForeground(const FGeometry& AllottedGeometry, FSlateWindowElementList& OutDrawElements, int32 LayerId) const
+{
+	TRACE_CPUPROFILER_EVENT_SCOPE(STerminal::PaintForeground);
+
+	// Selection highlight (translucent overlay on top of text).
+	if (bHasSelection || bSelecting)
+	{
+		const FLinearColor SelectionColor(SelectionHighlightColor.ReinterpretAsLinear() * FLinearColor(1.0f, 1.0f, 1.0f, 0.35f));
+		const int32 TotalRows = Buffer.GetTotalRows();
+		const int32 Columns = Buffer.GetColumns();
+
+		for (int32 ViewportRow = 0; ViewportRow < ViewportRows; ++ViewportRow)
+		{
+			const int32 AbsoluteRow = Buffer.GetAbsoluteRow(ViewportRow, ScrollOffset);
+			if (AbsoluteRow < 0 || AbsoluteRow >= TotalRows)
+			{
+				continue;
+			}
+
+			const float Y = ViewportRow * CellHeight;
+			int32 Col = 0;
+			while (Col < Columns)
+			{
+				if (!IsCellSelected(AbsoluteRow, Col))
+				{
+					Col++;
+					continue;
+				}
+
+				// Batch consecutive selected cells.
+				int32 RunEnd = Col + 1;
+				while (RunEnd < Columns && IsCellSelected(AbsoluteRow, RunEnd))
+				{
+					RunEnd++;
+				}
+
+				const float X = Col * CellWidth;
+				const float RunWidth = (RunEnd - Col) * CellWidth;
+				FSlateDrawElement::MakeBox(
+					OutDrawElements,
+					LayerId,
+					AllottedGeometry.ToPaintGeometry(FVector2D(RunWidth, CellHeight), FSlateLayoutTransform(FVector2D(X, Y))),
+					FAppStyle::GetBrush("WhiteBrush"),
+					ESlateDrawEffect::None,
+					SelectionColor);
+
+				Col = RunEnd;
+			}
+		}
+	}
+
+	// Cursor.
+	if (Parser.bCursorVisible && Buffer.Cursor.bBlinkOn && ScrollOffset == 0)
+	{
+		const float CursorX = Buffer.Cursor.Column * CellWidth;
+		const float CursorY = Buffer.Cursor.Row * CellHeight;
+
+		FSlateDrawElement::MakeBox(
+			OutDrawElements,
+			LayerId,
+			AllottedGeometry.ToPaintGeometry(FVector2D(CellWidth, CellHeight), FSlateLayoutTransform(FVector2D(CursorX, CursorY))),
+			FAppStyle::GetBrush("WhiteBrush"),
+			ESlateDrawEffect::None,
+			FLinearColor(CursorDisplayColor));
+	}
+
+	// Per-cell glyph boundaries: draw a uniquely colored box for each character.
+	if (CVarTerminalDebugGlyphs.GetValueOnGameThread())
+	{
+		const int32 TotalRows = Buffer.GetTotalRows();
+		const int32 Columns = Buffer.GetColumns();
+
+		for (int32 ViewportRow = 0; ViewportRow < ViewportRows; ++ViewportRow)
+		{
+			const int32 AbsoluteRow = Buffer.GetAbsoluteRow(ViewportRow, ScrollOffset);
+			if (AbsoluteRow < 0 || AbsoluteRow >= TotalRows)
+			{
+				continue;
+			}
+
+			const float Y = ViewportRow * CellHeight;
+
+			for (int32 Col = 0; Col < Columns; ++Col)
+			{
+				const FTerminalCell& Cell = Buffer.GetCell(AbsoluteRow, Col);
+				if (Cell.Character == TEXT(' '))
+				{
+					continue;
+				}
+
+				const FLinearColor GlyphColor = FLinearColor::MakeFromHSV8(static_cast<uint8>((Col * 47 + ViewportRow * 131) % 256), 200, 255).CopyWithNewOpacity(0.25f);
+				const FLinearColor GlyphBorderColor = GlyphColor.CopyWithNewOpacity(0.7f);
+
+				const FVector2D GlyphPosition(Col * CellWidth, Y);
+				const FVector2D GlyphSize(CellWidth, CellHeight);
+
+				// Fill.
+				FSlateDrawElement::MakeBox(
+					OutDrawElements,
+					LayerId + 1,
+					AllottedGeometry.ToPaintGeometry(GlyphSize, FSlateLayoutTransform(GlyphPosition)),
+					FAppStyle::GetBrush("WhiteBrush"),
+					ESlateDrawEffect::None,
+					GlyphColor);
+
+				// Left edge.
+				FSlateDrawElement::MakeBox(
+					OutDrawElements,
+					LayerId + 1,
+					AllottedGeometry.ToPaintGeometry(FVector2D(1.0f, CellHeight), FSlateLayoutTransform(GlyphPosition)),
+					FAppStyle::GetBrush("WhiteBrush"),
+					ESlateDrawEffect::None,
+					GlyphBorderColor);
+			}
+		}
+	}
+}
+
+FReply STerminal::OnKeyChar(const FGeometry& MyGeometry, const FCharacterEvent& InCharacterEvent)
+{
+	if (!Session.IsValid())
+	{
+		return FReply::Unhandled();
+	}
+
+	const TCHAR Character = InCharacterEvent.GetCharacter();
+
+	// Skip control characters  - they are handled in OnKeyDown (Enter, Tab, arrows, etc.).
+	if (Character < 0x20)
+	{
+		return FReply::Unhandled();
+	}
+
+	// Encode as UTF-8 and write to the session.
+	const FUtf8String CharUtf8(FString(1, &Character));
+	UE_LOG(LogTerminal, VeryVerbose, TEXT("OnKeyChar: U+%04X (%d bytes UTF-8)."), static_cast<uint32>(Character), CharUtf8.Len());
+	Session->WriteInput(MakeArrayView(reinterpret_cast<const uint8*>(*CharUtf8), CharUtf8.Len()));
+
+	// Reset scroll to bottom on input.
+	ScrollOffset = 0;
+	UpdateScrollBar();
+
+	return FReply::Handled();
+}
+
+FReply STerminal::OnKeyDown(const FGeometry& MyGeometry, const FKeyEvent& InKeyEvent)
+{
+	const FKey Key = InKeyEvent.GetKey();
+	const bool bCtrl = InKeyEvent.IsControlDown();
+	const bool bShift = InKeyEvent.IsShiftDown();
+
+	// Shift+Arrow: start or extend selection from the cursor position.
+	// This is handled before the session check because selection does not
+	// require a shell session.
+	if (bShift && !bCtrl && (Key == EKeys::Up || Key == EKeys::Down || Key == EKeys::Left || Key == EKeys::Right))
+	{
+		if (!bHasSelection)
+		{
+			// Start selection at cursor (cursor row is in viewport-local space with ScrollOffset=0).
+			const int32 CursorAbsoluteRow = Buffer.GetAbsoluteRow(Buffer.Cursor.Row, 0);
+			SelectionAnchor = FIntPoint(Buffer.Cursor.Column, CursorAbsoluteRow);
+			SelectionActive = SelectionAnchor;
+			bHasSelection = true;
+		}
+
+		if (Key == EKeys::Up)
+		{
+			if (SelectionActive.Y > 0)
+			{
+				SelectionActive.Y -= 1;
+				const int32 TopVisibleRow = Buffer.GetAbsoluteRow(0, ScrollOffset);
+				if (SelectionActive.Y < TopVisibleRow)
+				{
+					ScrollOffset = FMath::Min(ScrollOffset + 1, Buffer.GetScrollbackLength());
+					UpdateScrollBar();
+				}
+			}
+		}
+		else if (Key == EKeys::Down)
+		{
+			if (SelectionActive.Y < Buffer.GetTotalRows() - 1)
+			{
+				SelectionActive.Y += 1;
+				const int32 BottomVisibleRow = Buffer.GetAbsoluteRow(ViewportRows - 1, ScrollOffset);
+				if (SelectionActive.Y > BottomVisibleRow)
+				{
+					ScrollOffset = FMath::Max(ScrollOffset - 1, 0);
+					UpdateScrollBar();
+				}
+			}
+		}
+		else if (Key == EKeys::Left)
+		{
+			if (SelectionActive.X > 0)
+			{
+				SelectionActive.X -= 1;
+			}
+			else if (SelectionActive.Y > 0)
+			{
+				SelectionActive.Y -= 1;
+				SelectionActive.X = Buffer.GetColumns() - 1;
+				const int32 TopVisibleRow = Buffer.GetAbsoluteRow(0, ScrollOffset);
+				if (SelectionActive.Y < TopVisibleRow)
+				{
+					ScrollOffset = FMath::Min(ScrollOffset + 1, Buffer.GetScrollbackLength());
+					UpdateScrollBar();
+				}
+			}
+		}
+		else if (Key == EKeys::Right)
+		{
+			if (SelectionActive.X < Buffer.GetColumns() - 1)
+			{
+				SelectionActive.X += 1;
+			}
+			else if (SelectionActive.Y < Buffer.GetTotalRows() - 1)
+			{
+				SelectionActive.Y += 1;
+				SelectionActive.X = 0;
+				const int32 BottomVisibleRow = Buffer.GetAbsoluteRow(ViewportRows - 1, ScrollOffset);
+				if (SelectionActive.Y > BottomVisibleRow)
+				{
+					ScrollOffset = FMath::Max(ScrollOffset - 1, 0);
+					UpdateScrollBar();
+				}
+			}
+		}
+
+		Invalidate(EInvalidateWidgetReason::Paint);
+		return FReply::Handled();
+	}
+
+	if (!Session.IsValid())
+	{
+		return FReply::Unhandled();
+	}
+
+	TArray<uint8> SequenceData;
+
+	auto WriteSequence = [&SequenceData](const char* Sequence)
+	{
+		const int32 Length = FCStringAnsi::Strlen(Sequence);
+		SequenceData.Append(reinterpret_cast<const uint8*>(Sequence), Length);
+	};
+
+	// Ctrl+C: copy if selection, otherwise send interrupt.
+	if (bCtrl && Key == EKeys::C)
+	{
+		if (bHasSelection)
+		{
+			FString SelectedText = GetSelectedText();
+			SelectedText.ReplaceInline(TEXT("\r\n"), TEXT("\n"));
+			SelectedText.ReplaceInline(TEXT("\r"), TEXT("\n"));
+			SelectedText.ReplaceInline(TEXT("\n"), LINE_TERMINATOR);
+			FPlatformApplicationMisc::ClipboardCopy(*SelectedText);
+			ClearSelection();
+			return FReply::Handled();
+		}
+		SequenceData.Add(0x03); // ETX
+	}
+	else if (bCtrl && Key == EKeys::V)
+	{
+		FString ClipboardText;
+		FPlatformApplicationMisc::ClipboardPaste(ClipboardText);
+		if (!ClipboardText.IsEmpty())
+		{
+			ClipboardText.ReplaceInline(TEXT("\r\n"), TEXT("\n"));
+			ClipboardText.ReplaceInline(TEXT("\r"), TEXT("\n"));
+			const FUtf8String ClipboardUtf8(ClipboardText);
+			TArray<uint8> UTF8Data;
+			if (Parser.bBracketedPaste)
+			{
+				static const uint8 BracketOpen[] = { 0x1B, '[', '2', '0', '0', '~' };
+				UTF8Data.Append(BracketOpen, UE_ARRAY_COUNT(BracketOpen));
+			}
+			UTF8Data.Append(reinterpret_cast<const uint8*>(*ClipboardUtf8), ClipboardUtf8.Len());
+			if (Parser.bBracketedPaste)
+			{
+				static const uint8 BracketClose[] = { 0x1B, '[', '2', '0', '1', '~' };
+				UTF8Data.Append(BracketClose, UE_ARRAY_COUNT(BracketClose));
+			}
+			Session->WriteInput(UTF8Data);
+			ScrollOffset = 0;
+			UpdateScrollBar();
+		}
+		return FReply::Handled();
+	}
+	else if (bCtrl && Key == EKeys::D)
+	{
+		SequenceData.Add(0x04); // EOT
+	}
+	else if (bCtrl && Key == EKeys::Z)
+	{
+		SequenceData.Add(0x1A); // SUB
+	}
+	else if (bCtrl && Key == EKeys::L)
+	{
+		SequenceData.Add(0x0C); // FF (clear screen)
+	}
+	else if (Key == EKeys::Enter)
+	{
+		SequenceData.Add(0x0D); // CR
+	}
+	else if (Key == EKeys::Tab)
+	{
+		SequenceData.Add(0x09); // HT
+	}
+	else if (Key == EKeys::Escape)
+	{
+		SequenceData.Add(0x1B); // ESC
+	}
+	else if (Key == EKeys::BackSpace)
+	{
+		SequenceData.Add(0x7F); // DEL (most terminals send DEL for backspace)
+	}
+	else if (Key == EKeys::Up)
+	{
+		WriteSequence(Parser.bApplicationCursorKeys ? "\x1BOA" : "\x1B[A");
+	}
+	else if (Key == EKeys::Down)
+	{
+		WriteSequence(Parser.bApplicationCursorKeys ? "\x1BOB" : "\x1B[B");
+	}
+	else if (Key == EKeys::Right)
+	{
+		WriteSequence(Parser.bApplicationCursorKeys ? "\x1BOC" : "\x1B[C");
+	}
+	else if (Key == EKeys::Left)
+	{
+		WriteSequence(Parser.bApplicationCursorKeys ? "\x1BOD" : "\x1B[D");
+	}
+	else if (Key == EKeys::Home)
+	{
+		WriteSequence(Parser.bApplicationCursorKeys ? "\x1BOH" : "\x1B[H");
+	}
+	else if (Key == EKeys::End)
+	{
+		WriteSequence(Parser.bApplicationCursorKeys ? "\x1BOF" : "\x1B[F");
+	}
+	else if (Key == EKeys::Insert)
+	{
+		WriteSequence("\x1B[2~");
+	}
+	else if (Key == EKeys::Delete)
+	{
+		WriteSequence("\x1B[3~");
+	}
+	else if (Key == EKeys::PageUp)
+	{
+		WriteSequence("\x1B[5~");
+	}
+	else if (Key == EKeys::PageDown)
+	{
+		WriteSequence("\x1B[6~");
+	}
+	else if (Key == EKeys::F1) { WriteSequence("\x1BOP"); }
+	else if (Key == EKeys::F2) { WriteSequence("\x1BOQ"); }
+	else if (Key == EKeys::F3) { WriteSequence("\x1BOR"); }
+	else if (Key == EKeys::F4) { WriteSequence("\x1BOS"); }
+	else if (Key == EKeys::F5) { WriteSequence("\x1B[15~"); }
+	else if (Key == EKeys::F6) { WriteSequence("\x1B[17~"); }
+	else if (Key == EKeys::F7) { WriteSequence("\x1B[18~"); }
+	else if (Key == EKeys::F8) { WriteSequence("\x1B[19~"); }
+	else if (Key == EKeys::F9) { WriteSequence("\x1B[20~"); }
+	else if (Key == EKeys::F10) { WriteSequence("\x1B[21~"); }
+	else if (Key == EKeys::F11) { WriteSequence("\x1B[23~"); }
+	else if (Key == EKeys::F12) { WriteSequence("\x1B[24~"); }
+	else
+	{
+		// Consume all unrecognized keys so the editor does not process them
+		// as shortcuts (e.g., 'f' triggering viewport focus. Printable characters
+		// will arrive via OnKeyChar; modifier-only presses are harmlessly consumed
+		// since the terminal ignores them. When no session is active, the early
+		// return above already lets everything propagate.
+		return FReply::Handled();
+	}
+
+	if (SequenceData.Num() > 0)
+	{
+		UE_LOG(LogTerminal, VeryVerbose, TEXT("OnKeyDown: %s (%d bytes)."), *Key.ToString(), SequenceData.Num());
+		Session->WriteInput(SequenceData);
+		ScrollOffset = 0;
+		UpdateScrollBar();
+	}
+
+	return FReply::Handled();
+}
+
+FReply STerminal::OnMouseButtonDown(const FGeometry& MyGeometry, const FPointerEvent& MouseEvent)
+{
+	// Forward mouse press to the application when tracking is active.
+	if (Parser.MouseTrackingMode != FVTParser::EMouseTrackingMode::None)
+	{
+		if (Session.IsValid())
+		{
+			const FVector2D LocalPosition = MyGeometry.AbsoluteToLocal(MouseEvent.GetScreenSpacePosition());
+			const int32 Col = FMath::Clamp(FMath::FloorToInt32(static_cast<float>(LocalPosition.X) / CellWidth), 0, Buffer.GetColumns() - 1);
+			const int32 Row = FMath::Clamp(FMath::FloorToInt32(static_cast<float>(LocalPosition.Y) / CellHeight), 0, ViewportRows - 1);
+			int32 Button = 0;
+			if (MouseEvent.GetEffectingButton() == EKeys::MiddleMouseButton) { Button = 1; }
+			else if (MouseEvent.GetEffectingButton() == EKeys::RightMouseButton) { Button = 2; }
+			SendMouseEvent(Button, Col, Row, false);
+		}
+		return FReply::Handled().CaptureMouse(SharedThis(this));
+	}
+
+	if (MouseEvent.GetEffectingButton() == EKeys::LeftMouseButton)
+	{
+		const FIntPoint Cell = PixelToCell(MyGeometry, MyGeometry.AbsoluteToLocal(MouseEvent.GetScreenSpacePosition()));
+		const double CurrentTime = FPlatformTime::Seconds();
+
+		// Detect multi-click (double/triple).
+		if (CurrentTime - LastClickTime < 0.4 && Cell == LastClickCell)
+		{
+			ClickCount++;
+		}
+		else
+		{
+			ClickCount = 1;
+		}
+		LastClickTime = CurrentTime;
+		LastClickCell = Cell;
+
+		if (ClickCount == 2)
+		{
+			// Double-click: select word.
+			const int32 AbsRow = Cell.Y;
+			int32 WordStart = Cell.X;
+			int32 WordEnd = Cell.X;
+
+			// Expand backward.
+			while (WordStart > 0)
+			{
+				const FTerminalCell& PrevCell = Buffer.GetCell(AbsRow, WordStart - 1);
+				if (FChar::IsWhitespace(PrevCell.Character) || PrevCell.Character == TEXT('\0'))
+				{
+					break;
+				}
+				WordStart--;
+			}
+			// Expand forward.
+			while (WordEnd < Buffer.GetColumns() - 1)
+			{
+				const FTerminalCell& NextCell = Buffer.GetCell(AbsRow, WordEnd + 1);
+				if (FChar::IsWhitespace(NextCell.Character) || NextCell.Character == TEXT('\0'))
+				{
+					break;
+				}
+				WordEnd++;
+			}
+
+			SelectionAnchor = FIntPoint(WordStart, AbsRow);
+			SelectionActive = FIntPoint(WordEnd, AbsRow);
+			bHasSelection = true;
+			bSelecting = false;
+			Invalidate(EInvalidateWidgetReason::Paint);
+			return FReply::Handled().CaptureMouse(SharedThis(this));
+		}
+		else if (ClickCount >= 3)
+		{
+			// Triple-click: select entire row.
+			const int32 AbsRow = Cell.Y;
+			SelectionAnchor = FIntPoint(0, AbsRow);
+			SelectionActive = FIntPoint(Buffer.GetColumns() - 1, AbsRow);
+			bHasSelection = true;
+			bSelecting = false;
+			Invalidate(EInvalidateWidgetReason::Paint);
+			return FReply::Handled().CaptureMouse(SharedThis(this));
+		}
+
+		// Single click: start selection.
+		SelectionAnchor = Cell;
+		SelectionActive = Cell;
+		bSelecting = true;
+		bHasSelection = false;
+
+		return FReply::Handled().CaptureMouse(SharedThis(this));
+	}
+
+	return FReply::Unhandled();
+}
+
+FReply STerminal::OnMouseButtonUp(const FGeometry& MyGeometry, const FPointerEvent& MouseEvent)
+{
+	// Forward mouse release to the application when tracking is active.
+	if (Parser.MouseTrackingMode != FVTParser::EMouseTrackingMode::None)
+	{
+		if (Session.IsValid())
+		{
+			const FVector2D LocalPosition = MyGeometry.AbsoluteToLocal(MouseEvent.GetScreenSpacePosition());
+			const int32 Col = FMath::Clamp(FMath::FloorToInt32(static_cast<float>(LocalPosition.X) / CellWidth), 0, Buffer.GetColumns() - 1);
+			const int32 Row = FMath::Clamp(FMath::FloorToInt32(static_cast<float>(LocalPosition.Y) / CellHeight), 0, ViewportRows - 1);
+			int32 Button = 0;
+			if (MouseEvent.GetEffectingButton() == EKeys::MiddleMouseButton) { Button = 1; }
+			else if (MouseEvent.GetEffectingButton() == EKeys::RightMouseButton) { Button = 2; }
+			SendMouseEvent(Button, Col, Row, true);
+		}
+		return FReply::Handled().ReleaseMouseCapture();
+	}
+
+	if (MouseEvent.GetEffectingButton() == EKeys::LeftMouseButton && bSelecting)
+	{
+		bSelecting = false;
+		SelectionAutoScrollDirection = 0;
+		if (SelectionAnchor != SelectionActive)
+		{
+			bHasSelection = true;
+		}
+		else
+		{
+			bHasSelection = false;
+		}
+		Invalidate(EInvalidateWidgetReason::Paint);
+		return FReply::Handled().ReleaseMouseCapture();
+	}
+
+	return FReply::Handled().ReleaseMouseCapture();
+}
+
+FReply STerminal::OnMouseMove(const FGeometry& MyGeometry, const FPointerEvent& MouseEvent)
+{
+	// Forward mouse motion to the application when any-event tracking (1003) is active,
+	// or button-event tracking (1002) with a button held.
+	if (Parser.MouseTrackingMode != FVTParser::EMouseTrackingMode::None)
+	{
+		// Determine which button is held: 0=left, 1=middle, 2=right, 3=none.
+		int32 HeldButton = 3;
+		if (MouseEvent.IsMouseButtonDown(EKeys::LeftMouseButton)) { HeldButton = 0; }
+		else if (MouseEvent.IsMouseButtonDown(EKeys::MiddleMouseButton)) { HeldButton = 1; }
+		else if (MouseEvent.IsMouseButtonDown(EKeys::RightMouseButton)) { HeldButton = 2; }
+
+		if (Parser.MouseTrackingMode == FVTParser::EMouseTrackingMode::Any
+			|| (Parser.MouseTrackingMode == FVTParser::EMouseTrackingMode::ButtonEvent && HeldButton != 3))
+		{
+			if (Session.IsValid())
+			{
+				const FVector2D LocalPosition = MyGeometry.AbsoluteToLocal(MouseEvent.GetScreenSpacePosition());
+				const int32 Col = FMath::Clamp(FMath::FloorToInt32(static_cast<float>(LocalPosition.X) / CellWidth), 0, Buffer.GetColumns() - 1);
+				const int32 Row = FMath::Clamp(FMath::FloorToInt32(static_cast<float>(LocalPosition.Y) / CellHeight), 0, ViewportRows - 1);
+				const int32 MotionButton = 32 + HeldButton; // 32 = motion flag; button encoded in low bits
+				SendMouseEvent(MotionButton, Col, Row, false);
+			}
+			return FReply::Handled();
+		}
+	}
+
+	if (bSelecting)
+	{
+		const FVector2D LocalPosition = MyGeometry.AbsoluteToLocal(MouseEvent.GetScreenSpacePosition());
+		const float ViewportHeight = ViewportRows * CellHeight;
+		const int32 Col = FMath::Clamp(FMath::FloorToInt32(static_cast<float>(LocalPosition.X) / CellWidth), 0, Buffer.GetColumns() - 1);
+
+		if (LocalPosition.Y < 0.0f)
+		{
+			// Mouse above viewport: scroll up to show older content.
+			const int32 LinesToScroll = FMath::Clamp(FMath::CeilToInt32(-LocalPosition.Y / CellHeight), 1, 5);
+			const int32 MaxScrollback = Buffer.GetScrollbackLength();
+			ScrollOffset = FMath::Clamp(ScrollOffset + LinesToScroll, 0, MaxScrollback);
+			SelectionActive = FIntPoint(Col, Buffer.GetAbsoluteRow(0, ScrollOffset));
+			SelectionAutoScrollDirection = LinesToScroll;
+			SelectionAutoScrollAccumulator = 0.0f;
+			UpdateScrollBar();
+		}
+		else if (LocalPosition.Y >= ViewportHeight)
+		{
+			// Mouse below viewport: scroll down to show newer content.
+			const int32 LinesToScroll = FMath::Clamp(FMath::CeilToInt32((LocalPosition.Y - ViewportHeight) / CellHeight), 1, 5);
+			ScrollOffset = FMath::Max(0, ScrollOffset - LinesToScroll);
+			SelectionActive = FIntPoint(Col, Buffer.GetAbsoluteRow(ViewportRows - 1, ScrollOffset));
+			SelectionAutoScrollDirection = -LinesToScroll;
+			SelectionAutoScrollAccumulator = 0.0f;
+			UpdateScrollBar();
+		}
+		else
+		{
+			// Mouse inside viewport: normal selection update.
+			SelectionAutoScrollDirection = 0;
+			SelectionActive = PixelToCell(MyGeometry, LocalPosition);
+		}
+
+		Invalidate(EInvalidateWidgetReason::Paint);
+		return FReply::Handled();
+	}
+
+	return FReply::Unhandled();
+}
+
+FReply STerminal::OnMouseWheel(const FGeometry& MyGeometry, const FPointerEvent& MouseEvent)
+{
+	if (MouseEvent.IsControlDown())
+	{
+		UTerminalSettings* Settings = GetMutableDefault<UTerminalSettings>();
+		const int32 Delta = MouseEvent.GetWheelDelta() > 0 ? 1 : -1;
+		Settings->FontSize = FMath::Clamp(Settings->FontSize + Delta, 6, 72);
+		Settings->SaveConfig();
+		FontInfo = FSlateFontInfo(TerminalFont, static_cast<float>(Settings->FontSize));
+		FontInfo.FontFallback = EFontFallback::FF_Max;
+		LastAllocatedSize = FVector2D::ZeroVector; // Force cell remeasure on next paint.
+		Invalidate(EInvalidateWidgetReason::Layout);
+		return FReply::Handled();
+	}
+
+	// When mouse tracking is active, forward wheel events to the application instead of scrolling.
+	if (Parser.MouseTrackingMode != FVTParser::EMouseTrackingMode::None)
+	{
+		if (Session.IsValid())
+		{
+			const FVector2D LocalPosition = MyGeometry.AbsoluteToLocal(MouseEvent.GetScreenSpacePosition());
+			const int32 Col = FMath::Clamp(FMath::FloorToInt32(static_cast<float>(LocalPosition.X) / CellWidth), 0, Buffer.GetColumns() - 1);
+			const int32 Row = FMath::Clamp(FMath::FloorToInt32(static_cast<float>(LocalPosition.Y) / CellHeight), 0, ViewportRows - 1);
+			const int32 WheelButton = MouseEvent.GetWheelDelta() > 0 ? 64 : 65; // 64 = wheel up, 65 = wheel down
+			SendMouseEvent(WheelButton, Col, Row, false);
+		}
+		return FReply::Handled();
+	}
+
+	// Alternate screen has no scrollback - ignore scroll attempts.
+	if (Buffer.IsAlternateBufferActive())
+	{
+		return FReply::Handled();
+	}
+
+	const float Delta = MouseEvent.GetWheelDelta();
+	const int32 LinesToScroll = FMath::Clamp(FMath::RoundToInt32(Delta) * 3, -ViewportRows, ViewportRows);
+
+	const int32 MaxScrollback = Buffer.GetScrollbackLength();
+	ScrollOffset = FMath::Clamp(ScrollOffset + LinesToScroll, 0, MaxScrollback);
+
+	UpdateScrollBar();
+	Invalidate(EInvalidateWidgetReason::Paint);
+	return FReply::Handled();
+}
+
+FReply STerminal::OnMouseButtonDoubleClick(const FGeometry& MyGeometry, const FPointerEvent& MouseEvent)
+{
+	// Already handled via click counting in OnMouseButtonDown.
+	return FReply::Handled();
+}
+
+FReply STerminal::OnFocusReceived(const FGeometry& MyGeometry, const FFocusEvent& InFocusEvent)
+{
+	bHasFocus = true;
+	Buffer.Cursor.bBlinkOn = true;
+	Invalidate(EInvalidateWidgetReason::Paint);
+	return FReply::Handled();
+}
+
+void STerminal::OnFocusLost(const FFocusEvent& InFocusEvent)
+{
+	bHasFocus = false;
+	Buffer.Cursor.bBlinkOn = true;
+
+	// Clear active drag-selection to avoid ambiguous state.
+	bSelecting = false;
+	SelectionAutoScrollDirection = 0;
+
+	Invalidate(EInvalidateWidgetReason::Paint);
+}
+
+FCursorReply STerminal::OnCursorQuery(const FGeometry& MyGeometry, const FPointerEvent& CursorEvent) const
+{
+	return FCursorReply::Cursor(EMouseCursor::TextEditBeam);
+}
+
+EActiveTimerReturnType STerminal::OnTick(double InCurrentTime, float InDeltaTime)
+{
+	if (!bInitialized || !Session.IsValid())
+	{
+		return EActiveTimerReturnType::Continue;
+	}
+
+	// Continuous auto-scroll while drag-selecting outside the viewport.
+	if (bSelecting && SelectionAutoScrollDirection != 0)
+	{
+		SelectionAutoScrollAccumulator += InDeltaTime;
+		static constexpr float AutoScrollInterval = 0.05f; // ~20 lines/second
+		if (SelectionAutoScrollAccumulator >= AutoScrollInterval)
+		{
+			SelectionAutoScrollAccumulator -= AutoScrollInterval;
+			const int32 MaxScrollback = Buffer.GetScrollbackLength();
+			ScrollOffset = FMath::Clamp(ScrollOffset + SelectionAutoScrollDirection, 0, MaxScrollback);
+
+			if (SelectionAutoScrollDirection > 0)
+			{
+				// Scrolling up: extend selection to top of viewport.
+				SelectionActive.Y = Buffer.GetAbsoluteRow(0, ScrollOffset);
+			}
+			else
+			{
+				// Scrolling down: extend selection to bottom of viewport.
+				SelectionActive.Y = Buffer.GetAbsoluteRow(ViewportRows - 1, ScrollOffset);
+			}
+
+			UpdateScrollBar();
+			Invalidate(EInvalidateWidgetReason::Paint);
+		}
+	}
+
+	// Consume output from the terminal session.
+	TArray<uint8> Output = Session->ConsumeOutput();
+
+	// Prepend any pending data from previous ticks.
+	if (PendingParseData.Num() > 0)
+	{
+		PendingParseData.Append(Output);
+		Output = MoveTemp(PendingParseData);
+		PendingParseData.Reset();
+	}
+
+	if (Output.Num() > 0)
+	{
+		LastOutputTime = InCurrentTime;
+
+		if (!bReceivedFirstOutput)
+		{
+			bReceivedFirstOutput = true;
+			UE_LOG(LogTerminal, VeryVerbose, TEXT("OnTick: first output consumed (%d bytes). Cursor at (%d,%d)."), Output.Num(), Buffer.Cursor.Column, Buffer.Cursor.Row);
+		}
+
+		const int32 BytesToParse = FMath::Min(Output.Num(), ParseBudgetPerTick);
+		Parser.Parse(Output.GetData(), BytesToParse);
+		TotalBytesConsumed += BytesToParse;
+
+		UE_LOG(LogTerminal, VeryVerbose, TEXT("OnTick: parsed %d bytes (total consumed: %lld). Cursor at (%d,%d). Parser state: %s."), BytesToParse, TotalBytesConsumed, Buffer.Cursor.Column, Buffer.Cursor.Row, *Parser.GetDebugStateName());
+
+		// Hex dump early output to diagnose startup issues.
+		if (TotalBytesConsumed <= 4096)
+		{
+			FString HexDump;
+			HexDump.Reserve(BytesToParse * 3 + 1);
+			for (int32 i = 0; i < BytesToParse; ++i)
+			{
+				HexDump += FString::Printf(TEXT("%02X "), Output[i]);
+			}
+			UE_LOG(LogTerminal, VeryVerbose, TEXT("OnTick: raw hex [%d bytes]: %s"), BytesToParse, *HexDump);
+		}
+
+		// Defer remainder if we exceeded the budget.
+		if (BytesToParse < Output.Num())
+		{
+			PendingParseData.Append(Output.GetData() + BytesToParse, Output.Num() - BytesToParse);
+		}
+
+		// Drain any responses the parser generated (DA, DSR replies).
+		if (Parser.ResponseBuffer.Num() > 0)
+		{
+			Session->WriteInput(Parser.ResponseBuffer);
+			Parser.ResponseBuffer.Reset();
+		}
+
+		// Defer repaint while synchronized output (DEC mode 2026) is active - the app
+		// is mid-frame and will send ?2026l when done, at which point we paint once.
+		if (!Parser.bSynchronizedOutput)
+		{
+			UpdateScrollBar();
+			Invalidate(EInvalidateWidgetReason::Paint);
+		}
+	}
+	else if (LastOutputTime > 0.0 && InCurrentTime - LastHeartbeatLogTime >= 5.0)
+	{
+		LastHeartbeatLogTime = InCurrentTime;
+		const double SilenceDuration = InCurrentTime - LastOutputTime;
+		UE_LOG(LogTerminal, VeryVerbose, TEXT("OnTick: no output for %.1fs. Process running: %d. Total bytes consumed: %lld. HasFocus: %d."), SilenceDuration, Session->IsRunning() ? 1 : 0, TotalBytesConsumed, bHasFocus ? 1 : 0);
+	}
+
+	return EActiveTimerReturnType::Continue;
+}
+
+EActiveTimerReturnType STerminal::OnCursorBlink(double InCurrentTime, float InDeltaTime)
+{
+	if (!bInitialized || !bHasFocus)
+	{
+		return EActiveTimerReturnType::Continue;
+	}
+
+	Buffer.Cursor.bBlinkOn = !Buffer.Cursor.bBlinkOn;
+	Invalidate(EInvalidateWidgetReason::Paint);
+	return EActiveTimerReturnType::Continue;
+}
+
+void STerminal::MeasureCellSize(float LayoutScale)
+{
+	// Get the glyph advance directly from the font cache at the actual rendering scale.
+	// The element batcher renders text at the geometry's accumulated scale (DPI), so we must
+	// query metrics at that same scale and divide back to local space. Otherwise font hinting
+	// causes the rendered per-character advance to drift from the grid CellWidth.
+	const TSharedRef<FSlateFontCache> FontCache = FSlateApplication::Get().GetRenderer()->GetFontCache();
+	FCharacterList& CharacterList = FontCache->GetCharacterList(FontInfo, LayoutScale);
+	const FCharacterEntry& Entry = CharacterList.GetCharacter(TEXT('M'), FontInfo.FontFallback);
+
+	CellWidth = FMath::Max(1.0f, static_cast<float>(Entry.XAdvance) / LayoutScale);
+	CellHeight = FMath::Max(1.0f, static_cast<float>(CharacterList.GetMaxHeight()) / LayoutScale);
+}
+
+void STerminal::InitializeTerminal()
+{
+	const UTerminalSettings* Settings = GetDefault<UTerminalSettings>();
+	const int32 ScrollbackLimit = Settings ? Settings->ScrollbackLimit : 131072;
+
+	UE_LOG(LogTerminal, VeryVerbose, TEXT("InitializeTerminal: %dx%d, scrollback=%d."), ViewportColumns, ViewportRows, ScrollbackLimit);
+
+	Buffer.Initialize(ViewportColumns, ViewportRows, ScrollbackLimit);
+	Parser.SetBuffer(&Buffer);
+
+	// Apply color scheme.
+	if (GEditor)
+	{
+		if (UTerminalSubsystem* Subsystem = GEditor->GetEditorSubsystem<UTerminalSubsystem>())
+		{
+			ApplyColorScheme(Subsystem->GetActiveColorScheme());
+		}
+	}
+
+	// Initialize tab stops.
+	Parser.InitializeTabStops();
+
+	// Create the platform-specific PTY session at the correct dimensions.
+	FString SessionError;
+	Session = ITerminalSession::CreateForCurrentPlatform(SessionError);
+
+	if (!Session.IsValid() && !SessionError.IsEmpty())
+	{
+		UE_LOG(LogTerminal, Error, TEXT("%s"), *SessionError);
+
+		const FUtf8String DisplayErrorUtf8(FString::Printf(TEXT("\r\nError: %s\r\n"), *SessionError));
+		Parser.Parse(reinterpret_cast<const uint8*>(*DisplayErrorUtf8), DisplayErrorUtf8.Len());
+	}
+
+	if (Session.IsValid())
+	{
+		const FString ShellPath = Settings ? Settings->ShellExecutablePath : FString();
+		const FString WorkingDirectory = FPaths::ConvertRelativePathToFull(FPaths::RootDir());
+		Session->OnProcessExited.BindSP(this, &STerminal::HandleProcessExited);
+		if (!Session->Create(ShellPath, WorkingDirectory, ViewportColumns, ViewportRows))
+		{
+			UE_LOG(LogTerminal, Error, TEXT("Failed to create terminal session."));
+			Session.Reset();
+
+			const FUtf8String ErrorUtf8(TEXT("\r\nError: Failed to create terminal session. Check the output log for details.\r\n"));
+			Parser.Parse(reinterpret_cast<const uint8*>(*ErrorUtf8), ErrorUtf8.Len());
+		}
+	}
+
+	bInitialized = true;
+	UpdateScrollBar();
+}
+
+void STerminal::UpdateViewportDimensions(const FGeometry& AllottedGeometry)
+{
+	// Re-measure cell size at the actual DPI scale so the grid matches rendered glyph advances.
+	MeasureCellSize(AllottedGeometry.Scale);
+
+	const FVector2D Size = AllottedGeometry.GetLocalSize();
+	const int32 NewColumns = FMath::Max(1, FMath::FloorToInt32(static_cast<float>(Size.X) / CellWidth));
+	const int32 NewRows = FMath::Max(1, FMath::FloorToInt32(static_cast<float>(Size.Y) / CellHeight));
+
+	if (!bInitialized)
+	{
+		// Require minimum dimensions to avoid initializing ConPTY at unusable sizes
+		// (e.g. 1x1 during layout restoration or docking animation). The shell would
+		// produce garbled output in a micro-terminal. Initialization is deferred until
+		// a subsequent paint provides reasonable geometry.
+		static constexpr int32 MinimumColumns = 2;
+		static constexpr int32 MinimumRows = 1;
+		if (NewColumns < MinimumColumns || NewRows < MinimumRows)
+		{
+			UE_LOG(LogTerminal, VeryVerbose, TEXT("Deferring initialization: geometry %dx%d below minimum %dx%d."), NewColumns, NewRows, MinimumColumns, MinimumRows);
+			return;
+		}
+
+		UE_LOG(LogTerminal, VeryVerbose, TEXT("First valid geometry received: %dx%d. Initializing terminal."), NewColumns, NewRows);
+		ViewportColumns = NewColumns;
+		ViewportRows = NewRows;
+		InitializeTerminal();
+		return;
+	}
+
+	if (NewColumns != ViewportColumns || NewRows != ViewportRows)
+	{
+		UE_LOG(LogTerminal, VeryVerbose, TEXT("Viewport resized: %dx%d -> %dx%d."), ViewportColumns, ViewportRows, NewColumns, NewRows);
+
+		ViewportColumns = NewColumns;
+		ViewportRows = NewRows;
+		Buffer.Resize(ViewportColumns, ViewportRows);
+
+		if (Session.IsValid())
+		{
+			Session->Resize(ViewportColumns, ViewportRows);
+		}
+
+		// Re-initialize tab stops for new column count.
+		Parser.InitializeTabStops();
+		UpdateScrollBar();
+	}
+}
+
+FIntPoint STerminal::PixelToCell(const FGeometry& MyGeometry, const FVector2D& LocalPosition) const
+{
+	if (Buffer.GetColumns() <= 0 || ViewportRows <= 0)
+	{
+		return FIntPoint(0, 0);
+	}
+
+	const int32 Col = FMath::Clamp(FMath::FloorToInt32(static_cast<float>(LocalPosition.X) / CellWidth), 0, Buffer.GetColumns() - 1);
+	const int32 ViewportRow = FMath::Clamp(FMath::FloorToInt32(static_cast<float>(LocalPosition.Y) / CellHeight), 0, ViewportRows - 1);
+	const int32 AbsoluteRow = Buffer.GetAbsoluteRow(ViewportRow, ScrollOffset);
+	return FIntPoint(Col, AbsoluteRow);
+}
+
+bool STerminal::IsCellSelected(int32 AbsoluteRow, int32 Column) const
+{
+	if (!bHasSelection && !bSelecting)
+	{
+		return false;
+	}
+
+	// Normalize selection range so Start <= End.
+	FIntPoint Start = SelectionAnchor;
+	FIntPoint End = SelectionActive;
+	if (Start.Y > End.Y || (Start.Y == End.Y && Start.X > End.X))
+	{
+		Swap(Start, End);
+	}
+
+	if (AbsoluteRow < Start.Y || AbsoluteRow > End.Y)
+	{
+		return false;
+	}
+
+	if (AbsoluteRow == Start.Y && AbsoluteRow == End.Y)
+	{
+		return Column >= Start.X && Column <= End.X;
+	}
+
+	if (AbsoluteRow == Start.Y)
+	{
+		return Column >= Start.X;
+	}
+
+	if (AbsoluteRow == End.Y)
+	{
+		return Column <= End.X;
+	}
+
+	return true; // Rows between start and end are fully selected.
+}
+
+FString STerminal::GetSelectedText() const
+{
+	if (!bHasSelection)
+	{
+		return FString();
+	}
+
+	FIntPoint Start = SelectionAnchor;
+	FIntPoint End = SelectionActive;
+	if (Start.Y > End.Y || (Start.Y == End.Y && Start.X > End.X))
+	{
+		Swap(Start, End);
+	}
+
+	return Buffer.GetTextInRange(Start.Y, Start.X, End.Y, End.X);
+}
+
+void STerminal::ClearSelection()
+{
+	bSelecting = false;
+	bHasSelection = false;
+	SelectionAutoScrollDirection = 0;
+	SelectionAnchor = FIntPoint(-1, -1);
+	SelectionActive = FIntPoint(-1, -1);
+	Invalidate(EInvalidateWidgetReason::Paint);
+}
+
+void STerminal::ApplyColorScheme(const FTerminalColorScheme& Scheme)
+{
+	Buffer.DefaultCell.Foreground = Scheme.DefaultForeground.ToFColor(true);
+	Buffer.DefaultCell.Background = Scheme.DefaultBackground.ToFColor(true);
+	Parser.CurrentForeground = Buffer.DefaultCell.Foreground;
+	Parser.CurrentBackground = Buffer.DefaultCell.Background;
+
+	SelectionHighlightColor = Scheme.SelectionColor.ToFColor(true);
+	CursorDisplayColor = Scheme.CursorColor.ToFColor(true);
+	CursorDisplayColor.A = 180; // Translucent overlay for cursor visibility
+
+	// Copy the 16-color ANSI palette from the scheme.
+	const int32 PaletteCount = FMath::Min(Scheme.Palette.Num(), 16);
+	for (int32 Index = 0; Index < PaletteCount; ++Index)
+	{
+		Parser.AnsiPalette[Index] = Scheme.Palette[Index].ToFColor(true);
+	}
+}
+
+void STerminal::UpdateScrollBar()
+{
+	if (!ScrollBar.IsValid())
+	{
+		return;
+	}
+
+	// Hide scrollbar in alternate screen buffer - no scrollback exists.
+	if (Buffer.IsAlternateBufferActive())
+	{
+		ScrollBar->SetVisibility(EVisibility::Collapsed);
+		return;
+	}
+	ScrollBar->SetVisibility(EVisibility::Visible);
+
+	const int32 ScrollbackLength = Buffer.GetScrollbackLength();
+	if (ScrollbackLength <= 0)
+	{
+		ScrollBar->SetState(0.0f, 1.0f);
+		return;
+	}
+
+	const float TotalLines = static_cast<float>(ScrollbackLength + ViewportRows);
+	const float ThumbSize = static_cast<float>(ViewportRows) / TotalLines;
+	// ScrollOffset=0 means bottom (most recent), ScrollOffset=max means top (oldest).
+	// SScrollBar offset is in [0, 1-ThumbSize]: 0=top, 1-ThumbSize=bottom.
+	const float ScrollFraction = static_cast<float>(ScrollOffset) / static_cast<float>(ScrollbackLength);
+	const float ThumbOffset = (1.0f - ScrollFraction) * (1.0f - ThumbSize);
+	ScrollBar->SetState(ThumbOffset, ThumbSize);
+}
+
+void STerminal::HandleScrollBarScrolled(float InScrollOffsetFraction)
+{
+	if (Buffer.IsAlternateBufferActive())
+	{
+		return;
+	}
+
+	const int32 ScrollbackLength = Buffer.GetScrollbackLength();
+	if (ScrollbackLength <= 0)
+	{
+		return;
+	}
+
+	const float TotalLines = static_cast<float>(ScrollbackLength + ViewportRows);
+	const float ThumbSize = static_cast<float>(ViewportRows) / TotalLines;
+	// Invert: scrollbar offset 0 = top = max ScrollOffset, offset (1-ThumbSize) = bottom = ScrollOffset 0.
+	const float MaxOffset = FMath::Max(1.0f - ThumbSize, UE_SMALL_NUMBER);
+	const float ScrollFraction = 1.0f - InScrollOffsetFraction / MaxOffset;
+	ScrollOffset = FMath::Clamp(FMath::RoundToInt32(ScrollFraction * static_cast<float>(ScrollbackLength)), 0, ScrollbackLength);
+	UpdateScrollBar();
+	Invalidate(EInvalidateWidgetReason::Paint);
+}
+
+void STerminal::SendMouseEvent(int32 Button, int32 Column, int32 Row, bool bRelease)
+{
+	if (!Session.IsValid())
+	{
+		return;
+	}
+
+	// Terminal mouse coordinates are 1-based.
+	const int32 Col1 = Column + 1;
+	const int32 Row1 = Row + 1;
+
+	if (Parser.bSGRMouseEncoding)
+	{
+		// SGR format: CSI < Button ; Col ; Row M (press) or m (release).
+		const TCHAR Suffix = bRelease ? TEXT('m') : TEXT('M');
+		const FString Sequence = FString::Printf(TEXT("\x1B[<%d;%d;%d%c"), Button, Col1, Row1, Suffix);
+		const FTCHARToUTF8 Utf8(*Sequence);
+		Session->WriteInput(MakeArrayView(reinterpret_cast<const uint8*>(Utf8.Get()), Utf8.Length()));
+	}
+	else
+	{
+		// Legacy X10 format: ESC [ M Cb Cx Cy (all 1-based, offset by 32).
+		const uint8 Cb = static_cast<uint8>((bRelease ? 3 : Button) + 32);
+		const uint8 Cx = static_cast<uint8>(FMath::Min(Col1 + 32, 255));
+		const uint8 Cy = static_cast<uint8>(FMath::Min(Row1 + 32, 255));
+		const uint8 Bytes[] = { 0x1B, '[', 'M', Cb, Cx, Cy };
+		Session->WriteInput(Bytes);
+	}
+}
+
+void STerminal::HandleProcessExited(int32 ExitCode)
+{
+	UE_LOG(LogTerminal, Display, TEXT("Shell process exited with code %d."), ExitCode);
+
+	// Drain any remaining output the reader thread accumulated before the process exited.
+	// Without this, the shell's final output (error messages, etc.) would be lost when
+	// Session.Reset() destroys the staging buffer.
+	if (Session.IsValid())
+	{
+		const TArray<uint8> RemainingOutput = Session->ConsumeOutput();
+		if (RemainingOutput.Num() > 0)
+		{
+			Parser.Parse(RemainingOutput.GetData(), RemainingOutput.Num());
+		}
+	}
+
+	// Write a visual marker into the terminal so the user can see the exit.
+	const FUtf8String ExitUtf8(FString::Printf(TEXT("\r\n[Process exited with code %d]\r\n"), ExitCode));
+	Parser.Parse(reinterpret_cast<const uint8*>(*ExitUtf8), ExitUtf8.Len());
+
+	bSessionDead = true;
+	Session.Reset();
+	UpdateScrollBar();
+	Invalidate(EInvalidateWidgetReason::Paint);
+}
diff --git a/Engine/Plugins/Experimental/Terminal/Source/Terminal/Private/TerminalBuffer.cpp b/Engine/Plugins/Experimental/Terminal/Source/Terminal/Private/TerminalBuffer.cpp
new file mode 100644
--- /dev/null
+++ b/Engine/Plugins/Experimental/Terminal/Source/Terminal/Private/TerminalBuffer.cpp
@@ -0,0 +1,408 @@
+// Copyright Epic Games, Inc. All Rights Reserved.
+
+#include "TerminalBuffer.h"
+
+FTerminalBuffer::FTerminalBuffer()
+{
+}
+
+void FTerminalBuffer::Initialize(int32 InColumns, int32 InViewportRows, int32 InScrollbackLimit)
+{
+	Columns = FMath::Max(1, InColumns);
+	ViewportRows = FMath::Max(1, InViewportRows);
+	ScrollbackLimit = FMath::Max(0, InScrollbackLimit);
+
+	const int32 MaxRows = ViewportRows + ScrollbackLimit;
+	Rows.SetNum(MaxRows);
+	RingStart = 0;
+	TotalRowsWritten = ViewportRows;
+
+	for (int32 Index = 0; Index < ViewportRows; ++Index)
+	{
+		EnsureRow(Index);
+	}
+
+	Cursor = FTerminalCursor();
+}
+
+void FTerminalBuffer::Resize(int32 InColumns, int32 InViewportRows)
+{
+	if (InColumns == Columns && InViewportRows == ViewportRows)
+	{
+		return;
+	}
+
+	// Collect the current content into a linear array.
+	const int32 OldTotalRows = GetTotalRows();
+	TArray<TArray<FTerminalCell>> OldContent;
+	OldContent.Reserve(OldTotalRows);
+	for (int32 Index = 0; Index < OldTotalRows; ++Index)
+	{
+		const int32 RingIndex = GetRingIndex(Index);
+		if (Rows.IsValidIndex(RingIndex) && Rows[RingIndex].Num() > 0)
+		{
+			OldContent.Add(MoveTemp(Rows[RingIndex]));
+		}
+		else
+		{
+			TArray<FTerminalCell> BlankRow;
+			BlankRow.SetNum(Columns);
+			for (FTerminalCell& Cell : BlankRow)
+			{
+				Cell = DefaultCell;
+			}
+			OldContent.Add(MoveTemp(BlankRow));
+		}
+	}
+
+	// Reinitialize with new dimensions.
+	Columns = FMath::Max(1, InColumns);
+	ViewportRows = FMath::Max(1, InViewportRows);
+
+	const int32 MaxRows = ViewportRows + ScrollbackLimit;
+	Rows.SetNum(MaxRows);
+	RingStart = 0;
+
+	// Copy back as many rows as fit, adjusting column width.
+	const int32 RowsToCopy = FMath::Min(OldContent.Num(), MaxRows);
+	const int32 StartOffset = FMath::Max(0, OldContent.Num() - RowsToCopy);
+	TotalRowsWritten = FMath::Max(RowsToCopy, ViewportRows);
+
+	for (int32 Index = 0; Index < MaxRows; ++Index)
+	{
+		const int32 SourceIndex = StartOffset + Index;
+		if (Index < RowsToCopy && SourceIndex < OldContent.Num())
+		{
+			TArray<FTerminalCell>& SourceRow = OldContent[SourceIndex];
+			Rows[Index].SetNum(Columns);
+			for (int32 Col = 0; Col < Columns; ++Col)
+			{
+				Rows[Index][Col] = (Col < SourceRow.Num()) ? SourceRow[Col] : DefaultCell;
+			}
+		}
+		else
+		{
+			Rows[Index].SetNum(Columns);
+			for (FTerminalCell& Cell : Rows[Index])
+			{
+				Cell = DefaultCell;
+			}
+		}
+	}
+
+	// Clamp cursor.
+	Cursor.Row = FMath::Clamp(Cursor.Row, 0, ViewportRows - 1);
+	Cursor.Column = FMath::Clamp(Cursor.Column, 0, Columns - 1);
+}
+
+const FTerminalCell& FTerminalBuffer::GetCell(int32 AbsoluteRow, int32 Column) const
+{
+	static const FTerminalCell EmptyCell;
+
+	if (AbsoluteRow < 0 || AbsoluteRow >= GetTotalRows() || Column < 0 || Column >= Columns)
+	{
+		return EmptyCell;
+	}
+
+	const int32 RingIndex = GetRingIndex(AbsoluteRow);
+	if (Rows.IsValidIndex(RingIndex) && Rows[RingIndex].Num() > Column)
+	{
+		return Rows[RingIndex][Column];
+	}
+	return EmptyCell;
+}
+
+void FTerminalBuffer::SetCell(int32 AbsoluteRow, int32 Column, const FTerminalCell& Cell)
+{
+	if (AbsoluteRow < 0 || AbsoluteRow >= GetTotalRows() || Column < 0 || Column >= Columns)
+	{
+		return;
+	}
+
+	const int32 RingIndex = GetRingIndex(AbsoluteRow);
+	EnsureRow(RingIndex);
+	Rows[RingIndex][Column] = Cell;
+}
+
+void FTerminalBuffer::PushNewLine()
+{
+	const int32 MaxRows = ViewportRows + ScrollbackLimit;
+
+	if (TotalRowsWritten < MaxRows)
+	{
+		const int32 RingIndex = GetRingIndex(TotalRowsWritten);
+		EnsureRow(RingIndex);
+		// Clear the new row.
+		for (FTerminalCell& Cell : Rows[RingIndex])
+		{
+			Cell = DefaultCell;
+		}
+		TotalRowsWritten++;
+	}
+	else
+	{
+		// Ring buffer is full  - reuse the oldest row.
+		const int32 RingIndex = RingStart;
+		EnsureRow(RingIndex);
+		for (FTerminalCell& Cell : Rows[RingIndex])
+		{
+			Cell = DefaultCell;
+		}
+		RingStart = (RingStart + 1) % MaxRows;
+	}
+}
+
+int32 FTerminalBuffer::GetAbsoluteRow(int32 ViewportRow, int32 ScrollOffset) const
+{
+	const int32 ScrollbackLength = GetScrollbackLength();
+	const int32 BaseOffset = FMath::Clamp(ScrollOffset, 0, ScrollbackLength);
+	return (GetTotalRows() - ViewportRows - BaseOffset) + ViewportRow;
+}
+
+void FTerminalBuffer::Clear()
+{
+	const int32 MaxRows = ViewportRows + ScrollbackLimit;
+	RingStart = 0;
+	TotalRowsWritten = ViewportRows;
+
+	for (int32 Index = 0; Index < ViewportRows; ++Index)
+	{
+		EnsureRow(Index);
+		for (FTerminalCell& Cell : Rows[Index])
+		{
+			Cell = DefaultCell;
+		}
+	}
+
+	Cursor = FTerminalCursor();
+}
+
+int32 FTerminalBuffer::GetScrollbackLength() const
+{
+	return FMath::Max(0, GetTotalRows() - ViewportRows);
+}
+
+int32 FTerminalBuffer::GetTotalRows() const
+{
+	return TotalRowsWritten;
+}
+
+void FTerminalBuffer::ScrollRegionUp(int32 TopRow, int32 BottomRow, int32 Count)
+{
+	if (Count <= 0 || TopRow >= BottomRow)
+	{
+		return;
+	}
+
+	const int32 ScrollbackLength = GetScrollbackLength();
+
+	for (int32 Row = TopRow; Row <= BottomRow; ++Row)
+	{
+		const int32 SourceRow = Row + Count;
+		const int32 AbsoluteDestination = GetAbsoluteRow(Row, 0);
+
+		if (SourceRow <= BottomRow)
+		{
+			const int32 AbsoluteSource = GetAbsoluteRow(SourceRow, 0);
+			for (int32 Col = 0; Col < Columns; ++Col)
+			{
+				SetCell(AbsoluteDestination, Col, GetCell(AbsoluteSource, Col));
+			}
+		}
+		else
+		{
+			for (int32 Col = 0; Col < Columns; ++Col)
+			{
+				SetCell(AbsoluteDestination, Col, DefaultCell);
+			}
+		}
+	}
+}
+
+void FTerminalBuffer::ScrollRegionDown(int32 TopRow, int32 BottomRow, int32 Count)
+{
+	if (Count <= 0 || TopRow >= BottomRow)
+	{
+		return;
+	}
+
+	for (int32 Row = BottomRow; Row >= TopRow; --Row)
+	{
+		const int32 SourceRow = Row - Count;
+		const int32 AbsoluteDestination = GetAbsoluteRow(Row, 0);
+
+		if (SourceRow >= TopRow)
+		{
+			const int32 AbsoluteSource = GetAbsoluteRow(SourceRow, 0);
+			for (int32 Col = 0; Col < Columns; ++Col)
+			{
+				SetCell(AbsoluteDestination, Col, GetCell(AbsoluteSource, Col));
+			}
+		}
+		else
+		{
+			for (int32 Col = 0; Col < Columns; ++Col)
+			{
+				SetCell(AbsoluteDestination, Col, DefaultCell);
+			}
+		}
+	}
+}
+
+void FTerminalBuffer::FillRegion(int32 AbsoluteRowStart, int32 ColumnStart, int32 AbsoluteRowEnd, int32 ColumnEnd, const FTerminalCell& Cell)
+{
+	for (int32 Row = AbsoluteRowStart; Row <= AbsoluteRowEnd; ++Row)
+	{
+		const int32 ColBegin = (Row == AbsoluteRowStart) ? ColumnStart : 0;
+		const int32 ColFinish = (Row == AbsoluteRowEnd) ? ColumnEnd : Columns - 1;
+
+		for (int32 Col = ColBegin; Col <= ColFinish; ++Col)
+		{
+			SetCell(Row, Col, Cell);
+		}
+	}
+}
+
+FString FTerminalBuffer::GetTextInRange(int32 StartRow, int32 StartColumn, int32 EndRow, int32 EndColumn) const
+{
+	FString Result;
+
+	for (int32 Row = StartRow; Row <= EndRow; ++Row)
+	{
+		const int32 ColBegin = (Row == StartRow) ? StartColumn : 0;
+		const int32 ColFinish = (Row == EndRow) ? EndColumn : Columns - 1;
+
+		for (int32 Col = ColBegin; Col <= ColFinish; ++Col)
+		{
+			const FTerminalCell& Cell = GetCell(Row, Col);
+			Result.AppendChar(Cell.Character);
+		}
+
+		if (Row < EndRow)
+		{
+			// Trim trailing spaces from the row before adding newline.
+			while (Result.Len() > 0 && Result[Result.Len() - 1] == TEXT(' '))
+			{
+				Result.LeftChopInline(1);
+			}
+			Result.AppendChar(TEXT('\n'));
+		}
+	}
+
+	return Result;
+}
+
+void FTerminalBuffer::ClearScrollback()
+{
+	if (GetScrollbackLength() == 0)
+	{
+		return;
+	}
+
+	// Collect current viewport content.
+	TArray<TArray<FTerminalCell>> ViewportContent;
+	ViewportContent.SetNum(ViewportRows);
+	for (int32 Row = 0; Row < ViewportRows; ++Row)
+	{
+		const int32 AbsRow = GetAbsoluteRow(Row, 0);
+		const int32 RingIndex = GetRingIndex(AbsRow);
+		if (Rows.IsValidIndex(RingIndex) && Rows[RingIndex].Num() > 0)
+		{
+			ViewportContent[Row] = Rows[RingIndex];
+		}
+		else
+		{
+			ViewportContent[Row].SetNum(Columns);
+			for (FTerminalCell& Cell : ViewportContent[Row])
+			{
+				Cell = DefaultCell;
+			}
+		}
+	}
+
+	// Reset ring buffer with only viewport rows.
+	RingStart = 0;
+	TotalRowsWritten = ViewportRows;
+	for (int32 Index = 0; Index < ViewportRows; ++Index)
+	{
+		Rows[Index] = MoveTemp(ViewportContent[Index]);
+	}
+	for (int32 Index = ViewportRows; Index < Rows.Num(); ++Index)
+	{
+		Rows[Index].Reset();
+	}
+}
+
+void FTerminalBuffer::ActivateAlternateBuffer()
+{
+	if (bAlternateBufferActive)
+	{
+		return;
+	}
+
+	bAlternateBufferActive = true;
+	SavedMainRows = MoveTemp(Rows);
+	SavedRingStart = RingStart;
+	SavedTotalRowsWritten = TotalRowsWritten;
+	SavedCursor = Cursor;
+
+	// Initialize a clean alternate buffer (viewport-sized only, no scrollback).
+	RingStart = 0;
+	TotalRowsWritten = ViewportRows;
+	Rows.SetNum(ViewportRows);
+	for (int32 Index = 0; Index < ViewportRows; ++Index)
+	{
+		EnsureRow(Index);
+		for (FTerminalCell& Cell : Rows[Index])
+		{
+			Cell = DefaultCell;
+		}
+	}
+	Cursor = FTerminalCursor();
+}
+
+void FTerminalBuffer::DeactivateAlternateBuffer()
+{
+	if (!bAlternateBufferActive)
+	{
+		return;
+	}
+
+	bAlternateBufferActive = false;
+	Rows = MoveTemp(SavedMainRows);
+	RingStart = SavedRingStart;
+	TotalRowsWritten = SavedTotalRowsWritten;
+	Cursor = SavedCursor;
+}
+
+int32 FTerminalBuffer::GetRingIndex(int32 AbsoluteRow) const
+{
+	const int32 MaxRows = Rows.Num();
+	if (MaxRows == 0)
+	{
+		return 0;
+	}
+	return (RingStart + AbsoluteRow) % MaxRows;
+}
+
+void FTerminalBuffer::EnsureRow(int32 RingIndex)
+{
+	if (Rows.IsValidIndex(RingIndex) && Rows[RingIndex].Num() != Columns)
+	{
+		Rows[RingIndex].SetNum(Columns);
+		for (FTerminalCell& Cell : Rows[RingIndex])
+		{
+			Cell = DefaultCell;
+		}
+	}
+	else if (!Rows.IsValidIndex(RingIndex))
+	{
+		// Should not happen with proper initialization, but defensive.
+		Rows.SetNum(RingIndex + 1);
+		Rows[RingIndex].SetNum(Columns);
+		for (FTerminalCell& Cell : Rows[RingIndex])
+		{
+			Cell = DefaultCell;
+		}
+	}
+}
diff --git a/Engine/Plugins/Experimental/Terminal/Source/Terminal/Private/TerminalColorScheme.cpp b/Engine/Plugins/Experimental/Terminal/Source/Terminal/Private/TerminalColorScheme.cpp
new file mode 100644
--- /dev/null
+++ b/Engine/Plugins/Experimental/Terminal/Source/Terminal/Private/TerminalColorScheme.cpp
@@ -0,0 +1,90 @@
+// Copyright Epic Games, Inc. All Rights Reserved.
+
+#include "TerminalColorScheme.h"
+
+#include "TerminalUtilities.h"
+
+#include "Dom/JsonObject.h"
+#include "Serialization/JsonReader.h"
+#include "Serialization/JsonSerializer.h"
+
+FTerminalColorScheme FTerminalColorScheme::MakeDefault()
+{
+	FTerminalColorScheme Scheme;
+	Scheme.Name = TEXT("Default");
+	Scheme.DefaultForeground = TerminalUtilities::ParseHexColor(TEXT("#D4D4D4"));
+	Scheme.DefaultBackground = TerminalUtilities::ParseHexColor(TEXT("#1E1E1E"));
+	Scheme.CursorColor = TerminalUtilities::ParseHexColor(TEXT("#AEAFAD"));
+	Scheme.SelectionColor = TerminalUtilities::ParseHexColor(TEXT("#264F78"));
+	Scheme.Palette = {
+		TerminalUtilities::ParseHexColor(TEXT("#000000")), TerminalUtilities::ParseHexColor(TEXT("#CD3131")),
+		TerminalUtilities::ParseHexColor(TEXT("#0DBC79")), TerminalUtilities::ParseHexColor(TEXT("#E5E510")),
+		TerminalUtilities::ParseHexColor(TEXT("#2472C8")), TerminalUtilities::ParseHexColor(TEXT("#BC3FBC")),
+		TerminalUtilities::ParseHexColor(TEXT("#11A8CD")), TerminalUtilities::ParseHexColor(TEXT("#E5E5E5")),
+		TerminalUtilities::ParseHexColor(TEXT("#666666")), TerminalUtilities::ParseHexColor(TEXT("#F14C4C")),
+		TerminalUtilities::ParseHexColor(TEXT("#23D18B")), TerminalUtilities::ParseHexColor(TEXT("#F5F543")),
+		TerminalUtilities::ParseHexColor(TEXT("#3B8EEA")), TerminalUtilities::ParseHexColor(TEXT("#D670D6")),
+		TerminalUtilities::ParseHexColor(TEXT("#29B8DB")), TerminalUtilities::ParseHexColor(TEXT("#FFFFFF")),
+	};
+	return Scheme;
+}
+
+bool FTerminalColorScheme::FromJSON(const FString& JSONString, FTerminalColorScheme& OutScheme)
+{
+	TSharedPtr<FJsonObject> JsonObject;
+	TSharedRef<TJsonReader<>> Reader = TJsonReaderFactory<>::Create(JSONString);
+	if (!FJsonSerializer::Deserialize(Reader, JsonObject) || !JsonObject.IsValid())
+	{
+		return false;
+	}
+
+	if (!JsonObject->HasField(TEXT("name")) ||
+		!JsonObject->HasField(TEXT("defaultForeground")) ||
+		!JsonObject->HasField(TEXT("defaultBackground")) ||
+		!JsonObject->HasField(TEXT("cursorColor")) ||
+		!JsonObject->HasField(TEXT("selectionColor")))
+	{
+		return false;
+	}
+
+	if (!JsonObject->TryGetStringField(TEXT("name"), OutScheme.Name))
+	{
+		return false;
+	}
+
+	FString ForegroundHex, BackgroundHex, CursorHex, SelectionHex;
+	if (!JsonObject->TryGetStringField(TEXT("defaultForeground"), ForegroundHex) ||
+		!JsonObject->TryGetStringField(TEXT("defaultBackground"), BackgroundHex) ||
+		!JsonObject->TryGetStringField(TEXT("cursorColor"), CursorHex) ||
+		!JsonObject->TryGetStringField(TEXT("selectionColor"), SelectionHex))
+	{
+		return false;
+	}
+
+	OutScheme.DefaultForeground = TerminalUtilities::ParseHexColor(ForegroundHex);
+	OutScheme.DefaultBackground = TerminalUtilities::ParseHexColor(BackgroundHex);
+	OutScheme.CursorColor = TerminalUtilities::ParseHexColor(CursorHex);
+	OutScheme.SelectionColor = TerminalUtilities::ParseHexColor(SelectionHex);
+
+	OutScheme.Palette.Reset();
+	const TArray<TSharedPtr<FJsonValue>>* PaletteArray = nullptr;
+	if (JsonObject->TryGetArrayField(TEXT("palette"), PaletteArray) && PaletteArray)
+	{
+		for (const TSharedPtr<FJsonValue>& Value : *PaletteArray)
+		{
+			OutScheme.Palette.Add(TerminalUtilities::ParseHexColor(Value->AsString()));
+		}
+	}
+
+	// Ensure palette has exactly 16 entries.
+	while (OutScheme.Palette.Num() < 16)
+	{
+		OutScheme.Palette.Add(FLinearColor::White);
+	}
+	if (OutScheme.Palette.Num() > 16)
+	{
+		OutScheme.Palette.SetNum(16);
+	}
+
+	return true;
+}
diff --git a/Engine/Plugins/Experimental/Terminal/Source/Terminal/Private/TerminalModule.cpp b/Engine/Plugins/Experimental/Terminal/Source/Terminal/Private/TerminalModule.cpp
new file mode 100644
--- /dev/null
+++ b/Engine/Plugins/Experimental/Terminal/Source/Terminal/Private/TerminalModule.cpp
@@ -0,0 +1,92 @@
+// Copyright Epic Games, Inc. All Rights Reserved.
+
+#include "TerminalModule.h"
+
+#include "ITerminalSession.h"
+#include "STerminal.h"
+#include "TerminalSubsystem.h"
+#include "Editor.h"
+
+#include "Framework/Application/SlateApplication.h"
+#include "Modules/ModuleManager.h"
+#include "Widgets/Docking/SDockTab.h"
+#include "Widgets/SBoxPanel.h"
+#include "Widgets/Layout/SBox.h"
+#include "Widgets/Layout/SScrollBar.h"
+#include "WorkspaceMenuStructure.h"
+#include "WorkspaceMenuStructureModule.h"
+
+DEFINE_LOG_CATEGORY(LogTerminal);
+
+IMPLEMENT_MODULE(FTerminalModule, Terminal)
+
+#define LOCTEXT_NAMESPACE "FTerminalModule"
+
+static const FName TerminalTabName("Terminal");
+
+void FTerminalModule::StartupModule()
+{
+	FGlobalTabmanager::Get()->RegisterNomadTabSpawner(TerminalTabName, FOnSpawnTab::CreateRaw(this, &FTerminalModule::SpawnTerminalTab))
+		.SetDisplayName(LOCTEXT("Terminal", "Terminal"))
+		.SetGroup(WorkspaceMenu::GetMenuStructure().GetToolsCategory())
+		.SetReuseTabMethod(FOnFindTabToReuse::CreateLambda([](const FTabId&)
+		{
+			return TSharedPtr<SDockTab>();
+		}));
+}
+
+void FTerminalModule::ShutdownModule()
+{
+	FGlobalTabmanager::Get()->UnregisterNomadTabSpawner(TerminalTabName);
+}
+
+TSharedRef<SDockTab> FTerminalModule::SpawnTerminalTab(const FSpawnTabArgs& SpawnTabArgs)
+{
+	TSharedRef<SScrollBar> Scrollbar = SNew(SScrollBar)
+		.Orientation(Orient_Vertical)
+		.AlwaysShowScrollbar(true)
+		.Thickness(FVector2D(6.0f, 6.0f));
+
+	TSharedRef<STerminal> TerminalWidget = SNew(STerminal)
+		.ExternalScrollbar(Scrollbar);
+
+	TSharedRef<SDockTab> DockTab = SNew(SDockTab)
+		.TabRole(ETabRole::NomadTab)
+		[
+			SNew(SHorizontalBox)
+			+ SHorizontalBox::Slot()
+			.FillWidth(1.0f)
+			[
+				TerminalWidget
+			]
+			+ SHorizontalBox::Slot()
+			.AutoWidth()
+			[
+				Scrollbar
+			]
+		];
+
+	// Auto-focus the terminal widget when the user clicks on the tab header so it
+	// is immediately ready for keyboard input.  Only respond to UserClickedOnTab  -
+	// other causes (programmatic activation, layout restoration on startup) must be
+	// ignored because SetKeyboardFocus itself triggers tab activation, which might
+	// create a feedback loop.
+	TWeakPtr<STerminal> WeakTerminal = TerminalWidget;
+	DockTab->SetOnTabActivated(SDockTab::FOnTabActivatedCallback::CreateLambda(
+		[WeakTerminal](TSharedRef<SDockTab>, ETabActivationCause ActivationCause)
+		{
+			if (ActivationCause != ETabActivationCause::UserClickedOnTab)
+			{
+				return;
+			}
+
+			if (TSharedPtr<STerminal> Terminal = WeakTerminal.Pin())
+			{
+				FSlateApplication::Get().SetKeyboardFocus(Terminal);
+			}
+		}));
+
+	return DockTab;
+}
+
+#undef LOCTEXT_NAMESPACE
diff --git a/Engine/Plugins/Experimental/Terminal/Source/Terminal/Private/TerminalSubsystem.cpp b/Engine/Plugins/Experimental/Terminal/Source/Terminal/Private/TerminalSubsystem.cpp
new file mode 100644
--- /dev/null
+++ b/Engine/Plugins/Experimental/Terminal/Source/Terminal/Private/TerminalSubsystem.cpp
@@ -0,0 +1,89 @@
+// Copyright Epic Games, Inc. All Rights Reserved.
+
+#include "TerminalSubsystem.h"
+
+#include "ConPTYSession.h"
+#include "TerminalSettings.h"
+
+#include "Interfaces/IPluginManager.h"
+#include "Misc/FileHelper.h"
+#include "Misc/Paths.h"
+
+void UTerminalSubsystem::Initialize(FSubsystemCollectionBase& Collection)
+{
+	Super::Initialize(Collection);
+
+	DefaultScheme = FTerminalColorScheme::MakeDefault();
+	ReloadColorSchemes();
+}
+
+void UTerminalSubsystem::Deinitialize()
+{
+	ColorSchemes.Empty();
+
+	Super::Deinitialize();
+}
+
+FTerminalColorScheme UTerminalSubsystem::GetActiveColorScheme() const
+{
+	const UTerminalSettings* Settings = GetDefault<UTerminalSettings>();
+	const FString SchemeName = Settings ? Settings->ColorSchemeName : TEXT("Default");
+	return GetColorScheme(SchemeName);
+}
+
+FTerminalColorScheme UTerminalSubsystem::GetColorScheme(const FString& Name) const
+{
+	if (const FTerminalColorScheme* Found = ColorSchemes.Find(Name))
+	{
+		return *Found;
+	}
+
+	UE_LOG(LogTerminal, Warning, TEXT("Color scheme '%s' not found, using default."), *Name);
+	return DefaultScheme;
+}
+
+void UTerminalSubsystem::ReloadColorSchemes()
+{
+	ColorSchemes.Empty();
+
+	// Always register the hardcoded default.
+	ColorSchemes.Add(DefaultScheme.Name, DefaultScheme);
+
+	// Find the plugin's Config/ColorSchemes/ directory.
+	TSharedPtr<IPlugin> Plugin = IPluginManager::Get().FindPlugin(UE_PLUGIN_NAME);
+	if (!Plugin.IsValid())
+	{
+		return;
+	}
+
+	const FString ColorSchemesDirectory = FPaths::Combine(Plugin->GetBaseDir(), TEXT("Config"), TEXT("ColorSchemes"));
+	if (!FPaths::DirectoryExists(ColorSchemesDirectory))
+	{
+		return;
+	}
+
+	// Scan for JSON files.
+	TArray<FString> JSONFiles;
+	IFileManager::Get().FindFiles(JSONFiles, *FPaths::Combine(ColorSchemesDirectory, TEXT("*.json")), true, false);
+
+	for (const FString& FileName : JSONFiles)
+	{
+		const FString FilePath = FPaths::Combine(ColorSchemesDirectory, FileName);
+		FString JSONContent;
+		if (FFileHelper::LoadFileToString(JSONContent, *FilePath))
+		{
+			FTerminalColorScheme Scheme;
+			if (FTerminalColorScheme::FromJSON(JSONContent, Scheme) && !Scheme.Name.IsEmpty())
+			{
+				ColorSchemes.Add(Scheme.Name, Scheme);
+				UE_LOG(LogTerminal, Verbose, TEXT("Loaded color scheme: %s"), *Scheme.Name);
+			}
+			else
+			{
+				UE_LOG(LogTerminal, Warning, TEXT("Failed to parse color scheme file: %s"), *FilePath);
+			}
+		}
+	}
+
+	UE_LOG(LogTerminal, Display, TEXT("Loaded %d terminal color scheme(s)."), ColorSchemes.Num());
+}
diff --git a/Engine/Plugins/Experimental/Terminal/Source/Terminal/Private/TerminalUtilities.cpp b/Engine/Plugins/Experimental/Terminal/Source/Terminal/Private/TerminalUtilities.cpp
new file mode 100644
--- /dev/null
+++ b/Engine/Plugins/Experimental/Terminal/Source/Terminal/Private/TerminalUtilities.cpp
@@ -0,0 +1,169 @@
+// Copyright Epic Games, Inc. All Rights Reserved.
+
+#include "TerminalUtilities.h"
+
+#include "ITerminalSession.h"
+#include "Misc/Paths.h"
+
+bool TerminalUtilities::IsMonospaceSafe(TCHAR Character)
+{
+	return Character >= 0x20 && Character <= 0x7E;
+}
+
+FString TerminalUtilities::ResolveSystemFontPath(const FString& FontFamily)
+{
+	// Reject font family names that contain path separators or relative path
+	// components to prevent path traversal when constructing candidate paths.
+	if (FontFamily.IsEmpty()
+		|| FontFamily.Contains(TEXT("/"))
+		|| FontFamily.Contains(TEXT("\\"))
+		|| FontFamily.Contains(TEXT(".."))
+		|| FontFamily.Contains(TEXT(":")))
+	{
+		UE_LOG(LogTerminal, Warning, TEXT("Rejecting font family name with invalid characters: '%s'"), *FontFamily);
+		return FString();
+	}
+
+#if PLATFORM_WINDOWS
+	const FString WindowsFontsDirectory = FPlatformMisc::GetEnvironmentVariable(TEXT("WINDIR")) / TEXT("Fonts");
+	const FString FilePath = WindowsFontsDirectory / (FontFamily + TEXT(".ttf"));
+	return FPaths::FileExists(FilePath) ? FilePath : FString();
+#elif PLATFORM_MAC
+	static const TCHAR* SearchDirectories[] = {
+		TEXT("/System/Library/Fonts"),
+		TEXT("/Library/Fonts"),
+	};
+
+	// Also search ~/Library/Fonts.
+	const FString HomeDirectory = FPlatformMisc::GetEnvironmentVariable(TEXT("HOME"));
+	const FString UserFontsDirectory = HomeDirectory / TEXT("Library/Fonts");
+
+	TArray<const TCHAR*, TFixedAllocator<3>> AllDirectories;
+	for (const TCHAR* Directory : SearchDirectories)
+	{
+		AllDirectories.Add(Directory);
+	}
+	if (!HomeDirectory.IsEmpty())
+	{
+		AllDirectories.Add(*UserFontsDirectory);
+	}
+
+	for (const TCHAR* Directory : AllDirectories)
+	{
+		const FString FilePath = FString(Directory) / (FontFamily + TEXT(".ttf"));
+		if (FPaths::FileExists(FilePath))
+		{
+			return FilePath;
+		}
+		// macOS also uses .ttc (TrueType Collection) files.
+		const FString CollectionPath = FString(Directory) / (FontFamily + TEXT(".ttc"));
+		if (FPaths::FileExists(CollectionPath))
+		{
+			return CollectionPath;
+		}
+	}
+	return FString();
+#elif PLATFORM_LINUX
+	// Search common Linux font directories. Some fonts live in subdirectories
+	// (e.g. /usr/share/fonts/truetype/dejavu/DejaVuSansMono.ttf), so we try
+	// both flat and one level of common subdirectory structures.
+	const FString HomeDirectory = FPlatformMisc::GetEnvironmentVariable(TEXT("HOME"));
+
+	TArray<FString> SearchDirectories;
+	SearchDirectories.Add(TEXT("/usr/share/fonts"));
+	SearchDirectories.Add(TEXT("/usr/share/fonts/truetype"));
+	SearchDirectories.Add(TEXT("/usr/local/share/fonts"));
+	if (!HomeDirectory.IsEmpty())
+	{
+		SearchDirectories.Add(HomeDirectory / TEXT(".local/share/fonts"));
+	}
+
+	const FString FamilyLower = FontFamily.ToLower();
+
+	for (const FString& Directory : SearchDirectories)
+	{
+		const FString FilePath = Directory / (FontFamily + TEXT(".ttf"));
+		if (FPaths::FileExists(FilePath))
+		{
+			return FilePath;
+		}
+
+		// Many distros organize fonts into subdirectories named after the family.
+		// Try common patterns: truetype/<family>/, <family>/, TTF/.
+		static const TCHAR* SubDirectoryPatterns[] = { TEXT("truetype"), TEXT("TTF") };
+		for (const TCHAR* SubDirectory : SubDirectoryPatterns)
+		{
+			const FString SubPath = Directory / SubDirectory / (FontFamily + TEXT(".ttf"));
+			if (FPaths::FileExists(SubPath))
+			{
+				return SubPath;
+			}
+		}
+
+		// Also try: <directory>/<familyLower>/<FontFamily>.ttf (e.g. dejavu/DejaVuSansMono.ttf).
+		const FString FamilySubPath = Directory / FamilyLower / (FontFamily + TEXT(".ttf"));
+		if (FPaths::FileExists(FamilySubPath))
+		{
+			return FamilySubPath;
+		}
+	}
+	return FString();
+#else
+	return FString();
+#endif
+}
+
+TArrayView<const TCHAR* const> TerminalUtilities::GetFallbackFontNames()
+{
+#if PLATFORM_WINDOWS
+	static const TCHAR* Names[] = { TEXT("CascadiaMono"), TEXT("CascadiaCode"), TEXT("consola"), TEXT("cour") };
+#elif PLATFORM_MAC
+	static const TCHAR* Names[] = { TEXT("SFMono-Regular"), TEXT("Menlo-Regular"), TEXT("Monaco") };
+#else
+	static const TCHAR* Names[] = { TEXT("DejaVuSansMono"), TEXT("LiberationMono"), TEXT("UbuntuMono"), TEXT("DroidSansMono") };
+#endif
+	return MakeArrayView(Names);
+}
+
+TArrayView<const TCHAR* const> TerminalUtilities::GetSymbolFontNames()
+{
+#if PLATFORM_WINDOWS
+	static const TCHAR* Names[] = { TEXT("seguisym") };
+#elif PLATFORM_MAC
+	static const TCHAR* Names[] = { TEXT("Apple Symbols") };
+#else
+	static const TCHAR* Names[] = { TEXT("DejaVuSans"), TEXT("NotoSansSymbols2-Regular") };
+#endif
+	return MakeArrayView(Names);
+}
+
+TArrayView<const TCHAR* const> TerminalUtilities::GetEmojiFontNames()
+{
+#if PLATFORM_WINDOWS
+	static const TCHAR* Names[] = { TEXT("seguiemj") };
+#elif PLATFORM_MAC
+	static const TCHAR* Names[] = { TEXT("Apple Color Emoji") };
+#else
+	static const TCHAR* Names[] = { TEXT("NotoColorEmoji") };
+#endif
+	return MakeArrayView(Names);
+}
+
+FString TerminalUtilities::GetEngineFallbackFontPath()
+{
+	return FPaths::EngineContentDir() / TEXT("Slate/Fonts/DroidSansMono.ttf");
+}
+
+FLinearColor TerminalUtilities::ParseHexColor(const FString& HexString)
+{
+	FColor Color = FColor::White;
+	if (HexString.StartsWith(TEXT("#")) && HexString.Len() == 7)
+	{
+		Color = FColor::FromHex(HexString);
+	}
+	else if (HexString.Len() == 6)
+	{
+		Color = FColor::FromHex(HexString);
+	}
+	return FLinearColor(Color);
+}
diff --git a/Engine/Plugins/Experimental/Terminal/Source/Terminal/Private/VTParser.cpp b/Engine/Plugins/Experimental/Terminal/Source/Terminal/Private/VTParser.cpp
new file mode 100644
--- /dev/null
+++ b/Engine/Plugins/Experimental/Terminal/Source/Terminal/Private/VTParser.cpp
@@ -0,0 +1,1587 @@
+// Copyright Epic Games, Inc. All Rights Reserved.
+
+#include "VTParser.h"
+
+#include "TerminalBuffer.h"
+
+// ── Sequence accumulator caps ──
+// These limits prevent unbounded memory growth from pathological input.
+// Input exceeding a cap is silently discarded; the sequence still terminates normally.
+
+/**
+ * Max characters in an OSC payload. OSC 52 (clipboard) carries base64-encoded data that can
+ * include images  - a 1920x1080 screenshot at 4 bytes/pixel is ~8MB base64. xterm has no hard
+ * limit; xterm.js caps at 10MB; tmux at ~32KB. Allowing up to 16mb.
+ */
+static constexpr int32 MaxOSCPayloadLength = 16 * 1024 * 1024;
+
+/**
+ * Max digits per individual CSI parameter. Windows Console limits values to int16 (~5 digits);
+ * xterm.js allows up to 2^31-1 (~10 digits).
+ */
+static constexpr int32 MaxCSIParamDigits = 16;
+
+/**
+ * Max number of semicolon-separated CSI parameters. DEC VT510 allows 16; xterm uses ~30
+ * (NPARAM); xterm.js allows 32.
+ */
+static constexpr int32 MaxCSIParamCount = 32;
+
+/**
+ * Max intermediate bytes (0x20-0x2F) in a CSI or escape sequence. ECMA-48 specifies at most
+ * one intermediate byte per sequence in practice; two is the observed maximum. 4 is defensive.
+ */
+static constexpr int32 MaxIntermediateChars = 4;
+
+/** Default ANSI 16-color palette values. */
+static const FColor DefaultAnsiPalette[16] = {
+	FColor(0, 0, 0),       FColor(205, 49, 49),    FColor(13, 188, 121),  FColor(229, 229, 16),
+	FColor(36, 114, 200),   FColor(188, 63, 188),   FColor(17, 168, 205),  FColor(229, 229, 229),
+	FColor(102, 102, 102),  FColor(241, 76, 76),    FColor(35, 209, 139),  FColor(245, 245, 67),
+	FColor(59, 142, 234),   FColor(214, 112, 214),  FColor(41, 184, 219),  FColor(255, 255, 255),
+};
+
+FVTParser::FVTParser()
+{
+	Reset();
+}
+
+void FVTParser::SetBuffer(FTerminalBuffer* InBuffer)
+{
+	Buffer = InBuffer;
+}
+
+void FVTParser::Reset()
+{
+	State = EParserState::Ground;
+	CurrentForeground = FColor(212, 212, 212);
+	CurrentBackground = FColor(30, 30, 30);
+	CurrentAttributes = ETerminalAttribute::None;
+	FMemory::Memcpy(AnsiPalette, DefaultAnsiPalette, sizeof(AnsiPalette));
+	ScrollRegionTop = -1;
+	ScrollRegionBottom = -1;
+	bAutoWrap = true;
+	bCursorVisible = true;
+	bApplicationCursorKeys = false;
+	bBracketedPaste = false;
+	WindowTitle.Empty();
+	ResponseBuffer.Reset();
+	ClearSequenceState();
+	EscapeIntermediateChars.Empty();
+	OSCPayload.Empty();
+	UTF8Codepoint = 0;
+	UTF8Remaining = 0;
+	InitializeTabStops();
+}
+
+FString FVTParser::GetDebugStateName() const
+{
+	switch (State)
+	{
+	case EParserState::Ground:              return TEXT("Ground");
+	case EParserState::Escape:              return TEXT("Escape");
+	case EParserState::EscapeIntermediate:  return TEXT("EscapeIntermediate");
+	case EParserState::CSIEntry:            return TEXT("CSIEntry");
+	case EParserState::CSIParam:            return TEXT("CSIParam");
+	case EParserState::CSIIntermediate:     return TEXT("CSIIntermediate");
+	case EParserState::CSIIgnore:           return TEXT("CSIIgnore");
+	case EParserState::OSCString:           return TEXT("OSCString");
+	case EParserState::DCSEntry:            return TEXT("DCSEntry");
+	case EParserState::DCSParam:            return TEXT("DCSParam");
+	case EParserState::DCSIntermediate:     return TEXT("DCSIntermediate");
+	case EParserState::DCSPassthrough:      return TEXT("DCSPassthrough");
+	case EParserState::DCSIgnore:           return TEXT("DCSIgnore");
+	case EParserState::SOSPMAPCString:      return TEXT("SOSPMAPCString");
+	default:                                return TEXT("Unknown");
+	}
+}
+
+void FVTParser::ClearSequenceState()
+{
+	CSIParams.Reset();
+	CSIIntermediateChars.Empty();
+	CurrentParam.Empty();
+	CSIPrivatePrefix = 0;
+}
+
+bool FVTParser::TryExecuteC0(uint32 Codepoint)
+{
+	// C0 controls (0x00-0x1F) that should be executed even mid-sequence.
+	// CAN (0x18), SUB (0x1A), ESC (0x1B) are handled as "anywhere" transitions
+	// in ProcessCodepoint before reaching this function.
+	if (!Buffer)
+	{
+		return Codepoint < 0x20;
+	}
+
+	switch (Codepoint)
+	{
+	case 0x07: // BEL
+		return true;
+	case 0x08: // BS  - backspace
+		if (Buffer->Cursor.Column > 0)
+		{
+			Buffer->Cursor.Column--;
+		}
+		return true;
+	case 0x09: // HT  - horizontal tab
+	{
+		const int32 Columns = Buffer->GetColumns();
+		int32 NextTab = Buffer->Cursor.Column + 1;
+		while (NextTab < Columns)
+		{
+			if (NextTab < TabStops.Num() && TabStops[NextTab])
+			{
+				break;
+			}
+			NextTab++;
+		}
+		Buffer->Cursor.Column = FMath::Min(NextTab, Columns - 1);
+		return true;
+	}
+	case 0x0A: // LF  - line feed
+	case 0x0B: // VT  - vertical tab (treated as LF)
+	case 0x0C: // FF  - form feed (treated as LF)
+		LineFeed();
+		return true;
+	case 0x0D: // CR  - carriage return
+		Buffer->Cursor.Column = 0;
+		return true;
+	default:
+		// Other C0 controls (NUL, SOH, etc.)  - ignore.
+		return Codepoint < 0x20;
+	}
+}
+
+void FVTParser::Parse(const uint8* Data, int32 Length)
+{
+	if (!Buffer)
+	{
+		return;
+	}
+
+	for (int32 Index = 0; Index < Length; ++Index)
+	{
+		const uint8 Byte = Data[Index];
+
+		// UTF-8 decoding.
+		if (UTF8Remaining > 0)
+		{
+			if ((Byte & 0xC0) == 0x80)
+			{
+				UTF8Codepoint = (UTF8Codepoint << 6) | (Byte & 0x3F);
+				UTF8Remaining--;
+				if (UTF8Remaining == 0)
+				{
+					// Reject ill-formed UTF-8: overlong encodings, surrogates, out-of-range.
+					const bool bOverlong = UTF8Codepoint < UTF8MinCodepoint;
+					const bool bSurrogate = UTF8Codepoint >= 0xD800 && UTF8Codepoint <= 0xDFFF;
+					const bool bOutOfRange = UTF8Codepoint > 0x10FFFF;
+
+					if (!bOverlong && !bSurrogate && !bOutOfRange)
+					{
+						ProcessCodepoint(UTF8Codepoint);
+					}
+				}
+			}
+			else
+			{
+				// Invalid continuation byte  - reset and process as new byte.
+				UTF8Remaining = 0;
+				UTF8Codepoint = 0;
+				// Fall through to process this byte as a new sequence start.
+				goto NewByte;
+			}
+			continue;
+		}
+
+	NewByte:
+		if (Byte < 0x80)
+		{
+			ProcessCodepoint(Byte);
+		}
+		else if ((Byte & 0xE0) == 0xC0)
+		{
+			UTF8Codepoint = Byte & 0x1F;
+			UTF8Remaining = 1;
+			UTF8MinCodepoint = 0x80;
+		}
+		else if ((Byte & 0xF0) == 0xE0)
+		{
+			UTF8Codepoint = Byte & 0x0F;
+			UTF8Remaining = 2;
+			UTF8MinCodepoint = 0x800;
+		}
+		else if ((Byte & 0xF8) == 0xF0)
+		{
+			UTF8Codepoint = Byte & 0x07;
+			UTF8Remaining = 3;
+			UTF8MinCodepoint = 0x10000;
+		}
+		else
+		{
+			// Invalid byte  - skip.
+		}
+	}
+}
+
+void FVTParser::ProcessCodepoint(uint32 Codepoint)
+{
+	// ── "Anywhere" transitions  - highest priority, override all states ──
+
+	if (Codepoint == 0x1B) // ESC  - abort any sequence, begin new escape
+	{
+		ClearSequenceState();
+		EscapeIntermediateChars.Empty();
+		State = EParserState::Escape;
+		return;
+	}
+
+	if (Codepoint == 0x18 || Codepoint == 0x1A) // CAN, SUB  - abort to ground
+	{
+		ClearSequenceState();
+		OSCPayload.Empty();
+		State = EParserState::Ground;
+		return;
+	}
+
+	// 8-bit C1 controls (U+0080-U+009F)  - "anywhere" transitions.
+	// In a UTF-8 terminal these arrive as decoded codepoints, not raw bytes.
+	if (Codepoint >= 0x80 && Codepoint <= 0x9F)
+	{
+		switch (Codepoint)
+		{
+		case 0x90: // DCS
+			ClearSequenceState();
+			State = EParserState::DCSEntry;
+			return;
+		case 0x9B: // CSI
+			ClearSequenceState();
+			State = EParserState::CSIEntry;
+			return;
+		case 0x9D: // OSC
+			OSCPayload.Empty();
+			State = EParserState::OSCString;
+			return;
+		case 0x9C: // ST  - string terminator
+			State = EParserState::Ground;
+			return;
+		case 0x98: // SOS
+		case 0x9E: // PM
+		case 0x9F: // APC
+			State = EParserState::SOSPMAPCString;
+			return;
+		default:
+			// Other 8-bit C1 controls  - execute (treat as no-op) and go to ground.
+			State = EParserState::Ground;
+			return;
+		}
+	}
+
+	// ── State-specific dispatch ──
+
+	switch (State)
+	{
+	case EParserState::Ground:            HandleGround(Codepoint); break;
+	case EParserState::Escape:            HandleEscape(Codepoint); break;
+	case EParserState::EscapeIntermediate: HandleEscapeIntermediate(Codepoint); break;
+	case EParserState::CSIEntry:          HandleCSIEntry(Codepoint); break;
+	case EParserState::CSIParam:          HandleCSIParam(Codepoint); break;
+	case EParserState::CSIIntermediate:   HandleCSIIntermediate(Codepoint); break;
+	case EParserState::CSIIgnore:         HandleCSIIgnore(Codepoint); break;
+	case EParserState::OSCString:         HandleOSCString(Codepoint); break;
+	case EParserState::DCSEntry:          HandleDCSEntry(Codepoint); break;
+	case EParserState::DCSParam:          HandleDCSParam(Codepoint); break;
+	case EParserState::DCSIntermediate:   HandleDCSIntermediate(Codepoint); break;
+	case EParserState::DCSPassthrough:    HandleDCSPassthrough(Codepoint); break;
+	case EParserState::DCSIgnore:         HandleDCSIgnore(Codepoint); break;
+	case EParserState::SOSPMAPCString:    HandleSOSPMAPCString(Codepoint); break;
+	}
+}
+
+// ── Ground ──
+
+void FVTParser::HandleGround(uint32 Codepoint)
+{
+	// C0 controls (except ESC/CAN/SUB which are handled in ProcessCodepoint).
+	if (Codepoint < 0x20)
+	{
+		TryExecuteC0(Codepoint);
+		return;
+	}
+
+	// DEL  - ignore in ground.
+	if (Codepoint == 0x7F)
+	{
+		return;
+	}
+
+	// Printable characters (0x20+).
+	if (IsZeroWidthCodepoint(Codepoint))
+	{
+		return;
+	}
+	// TCHAR is 16-bit on Windows; codepoints above U+FFFF cannot be represented
+	// in a single TCHAR without surrogate pairs. Replace with U+FFFD to avoid truncation.
+	if (Codepoint > 0xFFFF)
+	{
+		PutCharacter(static_cast<TCHAR>(0xFFFD));
+	}
+	else
+	{
+		PutCharacter(static_cast<TCHAR>(Codepoint));
+	}
+}
+
+// ── Escape ──
+
+void FVTParser::HandleEscape(uint32 Codepoint)
+{
+	// C0 controls mid-escape are executed.
+	if (Codepoint < 0x20)
+	{
+		TryExecuteC0(Codepoint);
+		return;
+	}
+
+	// Intermediate bytes (0x20-0x2F) → collect, transition to EscapeIntermediate.
+	if (Codepoint >= 0x20 && Codepoint <= 0x2F)
+	{
+		EscapeIntermediateChars.Empty();
+		EscapeIntermediateChars.AppendChar(static_cast<TCHAR>(Codepoint));
+		State = EParserState::EscapeIntermediate;
+		return;
+	}
+
+	// DEL  - ignore.
+	if (Codepoint == 0x7F)
+	{
+		return;
+	}
+
+	// Dispatch based on the byte after ESC.
+	switch (Codepoint)
+	{
+	case '[': // CSI introducer (ESC [)
+		ClearSequenceState();
+		State = EParserState::CSIEntry;
+		break;
+	case ']': // OSC introducer (ESC ])
+		OSCPayload.Empty();
+		State = EParserState::OSCString;
+		break;
+	case 'P': // DCS introducer (ESC P)
+		ClearSequenceState();
+		State = EParserState::DCSEntry;
+		break;
+	case 'X': // SOS (ESC X)
+	case '^': // PM (ESC ^)
+	case '_': // APC (ESC _)
+		State = EParserState::SOSPMAPCString;
+		break;
+	case '7': // DECSC  - save cursor
+		Buffer->Cursor.SavedRow = Buffer->Cursor.Row;
+		Buffer->Cursor.SavedColumn = Buffer->Cursor.Column;
+		State = EParserState::Ground;
+		break;
+	case '8': // DECRC  - restore cursor
+		Buffer->Cursor.Row = Buffer->Cursor.SavedRow;
+		Buffer->Cursor.Column = Buffer->Cursor.SavedColumn;
+		State = EParserState::Ground;
+		break;
+	case 'D': // IND  - index (move cursor down, scroll if at bottom)
+		LineFeed();
+		State = EParserState::Ground;
+		break;
+	case 'E': // NEL  - next line (CR + LF)
+		Buffer->Cursor.Column = 0;
+		LineFeed();
+		State = EParserState::Ground;
+		break;
+	case 'M': // RI  - reverse index (move cursor up, scroll if at top)
+	{
+		const int32 Top = GetEffectiveScrollTop();
+		if (Buffer->Cursor.Row == Top)
+		{
+			Buffer->ScrollRegionDown(Top, GetEffectiveScrollBottom(), 1);
+		}
+		else if (Buffer->Cursor.Row > 0)
+		{
+			Buffer->Cursor.Row--;
+		}
+		State = EParserState::Ground;
+		break;
+	}
+	case 'c': // RIS  - full reset
+		Reset();
+		if (Buffer)
+		{
+			Buffer->Clear();
+		}
+		break;
+	case 'H': // HTS  - set tab stop at current column
+		if (Buffer->Cursor.Column < TabStops.Num())
+		{
+			TabStops[Buffer->Cursor.Column] = true;
+		}
+		State = EParserState::Ground;
+		break;
+	case '\\': // ST  - string terminator (terminates OSC/DCS/etc. via ESC \)
+		// When ESC aborts an OSC via the "anywhere" transition, the payload is still
+		// accumulated in OSCPayload. Dispatch it now that ST has arrived.
+		if (!OSCPayload.IsEmpty())
+		{
+			ExecuteOSC();
+			OSCPayload.Empty();
+		}
+		State = EParserState::Ground;
+		break;
+	default:
+		// All other bytes 0x30-0x7E are dispatched as escape sequences.
+		if (Codepoint >= 0x30 && Codepoint <= 0x7E)
+		{
+			ExecuteEscapeSequence(Codepoint);
+		}
+		State = EParserState::Ground;
+		break;
+	}
+}
+
+// ── EscapeIntermediate ──
+
+void FVTParser::HandleEscapeIntermediate(uint32 Codepoint)
+{
+	// C0 controls mid-escape are executed.
+	if (Codepoint < 0x20)
+	{
+		TryExecuteC0(Codepoint);
+		return;
+	}
+
+	// More intermediate bytes (0x20-0x2F)  - collect.
+	if (Codepoint >= 0x20 && Codepoint <= 0x2F)
+	{
+		if (EscapeIntermediateChars.Len() < MaxIntermediateChars)
+		{
+			EscapeIntermediateChars.AppendChar(static_cast<TCHAR>(Codepoint));
+		}
+		return;
+	}
+
+	// DEL  - ignore.
+	if (Codepoint == 0x7F)
+	{
+		return;
+	}
+
+	// Final byte (0x30-0x7E)  - dispatch the escape sequence with intermediates.
+	if (Codepoint >= 0x30 && Codepoint <= 0x7E)
+	{
+		ExecuteEscapeSequence(Codepoint);
+	}
+	State = EParserState::Ground;
+}
+
+// ── CSI Entry ──
+
+void FVTParser::HandleCSIEntry(uint32 Codepoint)
+{
+	// C0 controls  - execute inline.
+	if (Codepoint < 0x20)
+	{
+		TryExecuteC0(Codepoint);
+		return;
+	}
+
+	// Intermediate bytes (0x20-0x2F)  - collect, go to CSIIntermediate.
+	if (Codepoint >= 0x20 && Codepoint <= 0x2F)
+	{
+		if (CSIIntermediateChars.Len() < MaxIntermediateChars)
+		{
+			CSIIntermediateChars.AppendChar(static_cast<TCHAR>(Codepoint));
+		}
+		State = EParserState::CSIIntermediate;
+		return;
+	}
+
+	// Parameter digits and semicolons (0x30-0x39, 0x3B).
+	if ((Codepoint >= 0x30 && Codepoint <= 0x39) || Codepoint == 0x3B)
+	{
+		State = EParserState::CSIParam;
+		HandleCSIParam(Codepoint);
+		return;
+	}
+
+	// Colon (0x3A)  - sub-parameter separator. Treat like semicolon for modern SGR.
+	if (Codepoint == 0x3A)
+	{
+		State = EParserState::CSIParam;
+		HandleCSIParam(Codepoint);
+		return;
+	}
+
+	// Private parameter prefix bytes (0x3C-0x3F): '<', '=', '>', '?'.
+	if (Codepoint >= 0x3C && Codepoint <= 0x3F)
+	{
+		CSIPrivatePrefix = static_cast<TCHAR>(Codepoint);
+		State = EParserState::CSIParam;
+		return;
+	}
+
+	// Final byte (0x40-0x7E)  - dispatch immediately (no params).
+	if (Codepoint >= 0x40 && Codepoint <= 0x7E)
+	{
+		ExecuteCSI(Codepoint);
+		State = EParserState::Ground;
+		return;
+	}
+
+	// DEL  - ignore.
+	if (Codepoint == 0x7F)
+	{
+		return;
+	}
+}
+
+// ── CSI Param ──
+
+void FVTParser::HandleCSIParam(uint32 Codepoint)
+{
+	// C0 controls  - execute inline per spec.
+	if (Codepoint < 0x20)
+	{
+		TryExecuteC0(Codepoint);
+		return;
+	}
+
+	// Parameter digits (0x30-0x39).
+	if (Codepoint >= 0x30 && Codepoint <= 0x39)
+	{
+		if (CurrentParam.Len() < MaxCSIParamDigits)
+		{
+			CurrentParam.AppendChar(static_cast<TCHAR>(Codepoint));
+		}
+		return;
+	}
+
+	// Semicolons (0x3B) and colons (0x3A)  - parameter/sub-parameter separators.
+	// Colons are a modern extension (e.g., 38:2:R:G:B for SGR true-color).
+	// We treat them identically to semicolons for parameter splitting.
+	if (Codepoint == 0x3B || Codepoint == 0x3A)
+	{
+		if (CSIParams.Num() < MaxCSIParamCount)
+		{
+			CSIParams.Add(CurrentParam.IsEmpty() ? 0 : FCString::Atoi(*CurrentParam));
+		}
+		CurrentParam.Empty();
+		return;
+	}
+
+	// Private marker bytes (0x3C-0x3F) appearing AFTER the initial position is an error.
+	// Transition to CSIIgnore to consume the rest of the malformed sequence.
+	if (Codepoint >= 0x3C && Codepoint <= 0x3F)
+	{
+		State = EParserState::CSIIgnore;
+		return;
+	}
+
+	// Intermediate bytes (0x20-0x2F)  - collect and transition to CSIIntermediate.
+	if (Codepoint >= 0x20 && Codepoint <= 0x2F)
+	{
+		if (CSIParams.Num() < MaxCSIParamCount)
+		{
+			CSIParams.Add(CurrentParam.IsEmpty() ? 0 : FCString::Atoi(*CurrentParam));
+		}
+		CurrentParam.Empty();
+		if (CSIIntermediateChars.Len() < MaxIntermediateChars)
+		{
+			CSIIntermediateChars.AppendChar(static_cast<TCHAR>(Codepoint));
+		}
+		State = EParserState::CSIIntermediate;
+		return;
+	}
+
+	// Final byte (0x40-0x7E)  - flush last param, dispatch.
+	if (Codepoint >= 0x40 && Codepoint <= 0x7E)
+	{
+		if (CSIParams.Num() < MaxCSIParamCount)
+		{
+			CSIParams.Add(CurrentParam.IsEmpty() ? 0 : FCString::Atoi(*CurrentParam));
+		}
+		CurrentParam.Empty();
+		ExecuteCSI(Codepoint);
+		State = EParserState::Ground;
+		return;
+	}
+
+	// DEL (0x7F)  - ignore.
+}
+
+// ── CSI Intermediate ──
+
+void FVTParser::HandleCSIIntermediate(uint32 Codepoint)
+{
+	// C0 controls  - execute inline.
+	if (Codepoint < 0x20)
+	{
+		TryExecuteC0(Codepoint);
+		return;
+	}
+
+	// More intermediate bytes (0x20-0x2F)  - collect.
+	if (Codepoint >= 0x20 && Codepoint <= 0x2F)
+	{
+		if (CSIIntermediateChars.Len() < MaxIntermediateChars)
+		{
+			CSIIntermediateChars.AppendChar(static_cast<TCHAR>(Codepoint));
+		}
+		return;
+	}
+
+	// Parameter bytes (0x30-0x3F) after intermediates = error → CSIIgnore.
+	if (Codepoint >= 0x30 && Codepoint <= 0x3F)
+	{
+		State = EParserState::CSIIgnore;
+		return;
+	}
+
+	// Final byte (0x40-0x7E)  - dispatch.
+	if (Codepoint >= 0x40 && Codepoint <= 0x7E)
+	{
+		ExecuteCSI(Codepoint);
+		State = EParserState::Ground;
+		return;
+	}
+
+	// DEL (0x7F)  - ignore.
+}
+
+// ── CSI Ignore ──
+
+void FVTParser::HandleCSIIgnore(uint32 Codepoint)
+{
+	// C0 controls  - execute inline even in ignore state.
+	if (Codepoint < 0x20)
+	{
+		TryExecuteC0(Codepoint);
+		return;
+	}
+
+	// Consume everything (0x20-0x3F) without action.
+	if (Codepoint >= 0x20 && Codepoint <= 0x3F)
+	{
+		return;
+	}
+
+	// Final byte (0x40-0x7E)  - go to ground WITHOUT dispatching.
+	if (Codepoint >= 0x40 && Codepoint <= 0x7E)
+	{
+		State = EParserState::Ground;
+		return;
+	}
+
+	// DEL (0x7F)  - ignore.
+}
+
+// ── OSC String ──
+
+void FVTParser::HandleOSCString(uint32 Codepoint)
+{
+	// BEL (0x07) terminates OSC (xterm extension, universally supported).
+	if (Codepoint == 0x07)
+	{
+		ExecuteOSC();
+		State = EParserState::Ground;
+		return;
+	}
+
+	// C0 controls other than BEL  - ignore per spec.
+	if (Codepoint < 0x20)
+	{
+		return;
+	}
+
+	// DEL (0x7F)  - ignore per spec (some implementations allow it as data).
+	if (Codepoint == 0x7F)
+	{
+		return;
+	}
+
+	// ESC and 8-bit ST (0x9C) are handled by "anywhere" transitions in
+	// ProcessCodepoint. The ESC transition moves to Escape state; the subsequent
+	// '\' dispatches the accumulated OSC payload from HandleEscape's ST case.
+
+	// Collect printable bytes as OSC data, capped to prevent unbounded growth.
+	if (OSCPayload.Len() < MaxOSCPayloadLength)
+	{
+		OSCPayload.AppendChar(static_cast<TCHAR>(Codepoint));
+	}
+}
+
+// ── DCS Entry ──
+
+void FVTParser::HandleDCSEntry(uint32 Codepoint)
+{
+	// C0 controls  - IGNORED (not executed) in DCS states per spec.
+	if (Codepoint < 0x20)
+	{
+		return;
+	}
+
+	// Intermediate bytes (0x20-0x2F)  - collect, go to DCSIntermediate.
+	if (Codepoint >= 0x20 && Codepoint <= 0x2F)
+	{
+		if (CSIIntermediateChars.Len() < MaxIntermediateChars)
+		{
+			CSIIntermediateChars.AppendChar(static_cast<TCHAR>(Codepoint));
+		}
+		State = EParserState::DCSIntermediate;
+		return;
+	}
+
+	// Parameter digits and semicolons (0x30-0x39, 0x3B).
+	if ((Codepoint >= 0x30 && Codepoint <= 0x39) || Codepoint == 0x3B)
+	{
+		State = EParserState::DCSParam;
+		HandleDCSParam(Codepoint);
+		return;
+	}
+
+	// Colon (0x3A)  - error → DCSIgnore.
+	if (Codepoint == 0x3A)
+	{
+		State = EParserState::DCSIgnore;
+		return;
+	}
+
+	// Private prefix bytes (0x3C-0x3F)  - collect, go to DCSParam.
+	if (Codepoint >= 0x3C && Codepoint <= 0x3F)
+	{
+		CSIPrivatePrefix = static_cast<TCHAR>(Codepoint);
+		State = EParserState::DCSParam;
+		return;
+	}
+
+	// Final byte (0x40-0x7E)  - go directly to DCSPassthrough (hook).
+	if (Codepoint >= 0x40 && Codepoint <= 0x7E)
+	{
+		State = EParserState::DCSPassthrough;
+		return;
+	}
+
+	// DEL (0x7F)  - ignore.
+}
+
+// ── DCS Param ──
+
+void FVTParser::HandleDCSParam(uint32 Codepoint)
+{
+	// C0 controls  - ignored in DCS.
+	if (Codepoint < 0x20)
+	{
+		return;
+	}
+
+	// Parameter digits (0x30-0x39).
+	if (Codepoint >= 0x30 && Codepoint <= 0x39)
+	{
+		if (CurrentParam.Len() < MaxCSIParamDigits)
+		{
+			CurrentParam.AppendChar(static_cast<TCHAR>(Codepoint));
+		}
+		return;
+	}
+
+	// Semicolons (0x3B).
+	if (Codepoint == 0x3B)
+	{
+		if (CSIParams.Num() < MaxCSIParamCount)
+		{
+			CSIParams.Add(CurrentParam.IsEmpty() ? 0 : FCString::Atoi(*CurrentParam));
+		}
+		CurrentParam.Empty();
+		return;
+	}
+
+	// Colon (0x3A)  - error → DCSIgnore.
+	if (Codepoint == 0x3A)
+	{
+		State = EParserState::DCSIgnore;
+		return;
+	}
+
+	// Private marker bytes (0x3C-0x3F) after initial position = error → DCSIgnore.
+	if (Codepoint >= 0x3C && Codepoint <= 0x3F)
+	{
+		State = EParserState::DCSIgnore;
+		return;
+	}
+
+	// Intermediate bytes (0x20-0x2F)  - collect, go to DCSIntermediate.
+	if (Codepoint >= 0x20 && Codepoint <= 0x2F)
+	{
+		if (CSIParams.Num() < MaxCSIParamCount)
+		{
+			CSIParams.Add(CurrentParam.IsEmpty() ? 0 : FCString::Atoi(*CurrentParam));
+		}
+		CurrentParam.Empty();
+		if (CSIIntermediateChars.Len() < MaxIntermediateChars)
+		{
+			CSIIntermediateChars.AppendChar(static_cast<TCHAR>(Codepoint));
+		}
+		State = EParserState::DCSIntermediate;
+		return;
+	}
+
+	// Final byte (0x40-0x7E)  - flush, go to DCSPassthrough.
+	if (Codepoint >= 0x40 && Codepoint <= 0x7E)
+	{
+		if (CSIParams.Num() < MaxCSIParamCount)
+		{
+			CSIParams.Add(CurrentParam.IsEmpty() ? 0 : FCString::Atoi(*CurrentParam));
+		}
+		CurrentParam.Empty();
+		State = EParserState::DCSPassthrough;
+		return;
+	}
+
+	// DEL (0x7F)  - ignore.
+}
+
+// ── DCS Intermediate ──
+
+void FVTParser::HandleDCSIntermediate(uint32 Codepoint)
+{
+	// C0 controls  - ignored in DCS.
+	if (Codepoint < 0x20)
+	{
+		return;
+	}
+
+	// More intermediate bytes (0x20-0x2F)  - collect.
+	if (Codepoint >= 0x20 && Codepoint <= 0x2F)
+	{
+		if (CSIIntermediateChars.Len() < MaxIntermediateChars)
+		{
+			CSIIntermediateChars.AppendChar(static_cast<TCHAR>(Codepoint));
+		}
+		return;
+	}
+
+	// Parameter bytes (0x30-0x3F) after intermediates = error → DCSIgnore.
+	if (Codepoint >= 0x30 && Codepoint <= 0x3F)
+	{
+		State = EParserState::DCSIgnore;
+		return;
+	}
+
+	// Final byte (0x40-0x7E)  - go to DCSPassthrough.
+	if (Codepoint >= 0x40 && Codepoint <= 0x7E)
+	{
+		State = EParserState::DCSPassthrough;
+		return;
+	}
+
+	// DEL (0x7F)  - ignore.
+}
+
+// ── DCS Passthrough ──
+
+void FVTParser::HandleDCSPassthrough(uint32 Codepoint)
+{
+	// C0 controls are passed through to the DCS handler (the "put" action).
+	// Since we don't act on DCS data, we silently consume everything.
+	// Termination is via "anywhere" transitions (ESC or 8-bit ST).
+
+	// DEL (0x7F)  - ignore.
+	// All other bytes (0x00-0x7E)  - silently consume.
+}
+
+// ── DCS Ignore ──
+
+void FVTParser::HandleDCSIgnore(uint32 Codepoint)
+{
+	// Silently consume all bytes until terminated by "anywhere" transition.
+	// C0 controls are ignored. All other bytes (0x20-0x7F) are ignored.
+}
+
+// ── SOS/PM/APC String ──
+
+void FVTParser::HandleSOSPMAPCString(uint32 Codepoint)
+{
+	// Silently consume all data until terminated by "anywhere" transition
+	// (ESC triggers Escape state, then '\' = ST; or 8-bit 0x9C = ST).
+}
+
+// ── Escape Sequence Dispatch (with intermediates) ──
+
+void FVTParser::ExecuteEscapeSequence(uint32 FinalByte)
+{
+	// This handles escape sequences that went through EscapeIntermediate,
+	// or single-character escapes not caught by HandleEscape's switch.
+	// Most escape-with-intermediate sequences are charset designations
+	// (ESC ( B, ESC ) 0, etc.) which we silently ignore for now.
+}
+
+// ── CSI Dispatch ──
+
+void FVTParser::ExecuteCSI(uint32 FinalByte)
+{
+	if (CSIPrivatePrefix != 0)
+	{
+		ExecutePrivateMode(FinalByte);
+		return;
+	}
+
+	switch (FinalByte)
+	{
+	case 'A': // CUU  - cursor up
+	{
+		const int32 Count = FMath::Max(1, GetCSIParam(0, 1));
+		Buffer->Cursor.Row = FMath::Max(GetEffectiveScrollTop(), Buffer->Cursor.Row - Count);
+		break;
+	}
+	case 'B': // CUD  - cursor down
+	{
+		const int32 Count = FMath::Max(1, GetCSIParam(0, 1));
+		Buffer->Cursor.Row = FMath::Min(GetEffectiveScrollBottom(), Buffer->Cursor.Row + Count);
+		break;
+	}
+	case 'C': // CUF  - cursor forward
+	{
+		const int32 Count = FMath::Max(1, GetCSIParam(0, 1));
+		Buffer->Cursor.Column = FMath::Min(Buffer->GetColumns() - 1, Buffer->Cursor.Column + Count);
+		break;
+	}
+	case 'D': // CUB  - cursor backward
+	{
+		const int32 Count = FMath::Max(1, GetCSIParam(0, 1));
+		Buffer->Cursor.Column = FMath::Max(0, Buffer->Cursor.Column - Count);
+		break;
+	}
+	case 'E': // CNL  - cursor next line
+	{
+		const int32 Count = FMath::Max(1, GetCSIParam(0, 1));
+		Buffer->Cursor.Row = FMath::Min(GetEffectiveScrollBottom(), Buffer->Cursor.Row + Count);
+		Buffer->Cursor.Column = 0;
+		break;
+	}
+	case 'F': // CPL  - cursor previous line
+	{
+		const int32 Count = FMath::Max(1, GetCSIParam(0, 1));
+		Buffer->Cursor.Row = FMath::Max(GetEffectiveScrollTop(), Buffer->Cursor.Row - Count);
+		Buffer->Cursor.Column = 0;
+		break;
+	}
+	case 'H': // CUP  - cursor position
+	case 'f': // HVP  - same as CUP
+	{
+		const int32 Row = FMath::Max(1, GetCSIParam(0, 1)) - 1; // 1-based to 0-based
+		const int32 Col = FMath::Max(1, GetCSIParam(1, 1)) - 1;
+		Buffer->Cursor.Row = FMath::Clamp(Row, 0, Buffer->GetViewportRows() - 1);
+		Buffer->Cursor.Column = FMath::Clamp(Col, 0, Buffer->GetColumns() - 1);
+		break;
+	}
+	case 'J': // ED  - erase in display
+	{
+		const int32 Mode = GetCSIParam(0, 0);
+		const int32 CursorAbsoluteRow = Buffer->GetAbsoluteRow(Buffer->Cursor.Row, 0);
+		FTerminalCell BlankCell = Buffer->DefaultCell;
+		BlankCell.Foreground = CurrentForeground;
+		BlankCell.Background = CurrentBackground;
+
+		if (Mode == 0) // Erase below (cursor to end)
+		{
+			// Current line from cursor to end.
+			for (int32 Col = Buffer->Cursor.Column; Col < Buffer->GetColumns(); ++Col)
+			{
+				Buffer->SetCell(CursorAbsoluteRow, Col, BlankCell);
+			}
+			// Lines below cursor.
+			for (int32 Row = Buffer->Cursor.Row + 1; Row < Buffer->GetViewportRows(); ++Row)
+			{
+				const int32 AbsRow = Buffer->GetAbsoluteRow(Row, 0);
+				for (int32 Col = 0; Col < Buffer->GetColumns(); ++Col)
+				{
+					Buffer->SetCell(AbsRow, Col, BlankCell);
+				}
+			}
+		}
+		else if (Mode == 1) // Erase above (start to cursor)
+		{
+			for (int32 Row = 0; Row < Buffer->Cursor.Row; ++Row)
+			{
+				const int32 AbsRow = Buffer->GetAbsoluteRow(Row, 0);
+				for (int32 Col = 0; Col < Buffer->GetColumns(); ++Col)
+				{
+					Buffer->SetCell(AbsRow, Col, BlankCell);
+				}
+			}
+			for (int32 Col = 0; Col <= Buffer->Cursor.Column; ++Col)
+			{
+				Buffer->SetCell(CursorAbsoluteRow, Col, BlankCell);
+			}
+		}
+		else if (Mode == 2) // Erase all  - push viewport into scrollback, then blank
+		{
+			for (int32 Row = 0; Row < Buffer->GetViewportRows(); ++Row)
+			{
+				Buffer->PushNewLine();
+			}
+		}
+		else if (Mode == 3) // Erase saved lines  - clear scrollback
+		{
+			Buffer->ClearScrollback();
+		}
+		break;
+	}
+	case 'K': // EL  - erase in line
+	{
+		const int32 Mode = GetCSIParam(0, 0);
+		const int32 CursorAbsoluteRow = Buffer->GetAbsoluteRow(Buffer->Cursor.Row, 0);
+		FTerminalCell BlankCell = Buffer->DefaultCell;
+		BlankCell.Foreground = CurrentForeground;
+		BlankCell.Background = CurrentBackground;
+
+		if (Mode == 0) // Erase to right
+		{
+			for (int32 Col = Buffer->Cursor.Column; Col < Buffer->GetColumns(); ++Col)
+			{
+				Buffer->SetCell(CursorAbsoluteRow, Col, BlankCell);
+			}
+		}
+		else if (Mode == 1) // Erase to left
+		{
+			for (int32 Col = 0; Col <= Buffer->Cursor.Column; ++Col)
+			{
+				Buffer->SetCell(CursorAbsoluteRow, Col, BlankCell);
+			}
+		}
+		else if (Mode == 2) // Erase entire line
+		{
+			for (int32 Col = 0; Col < Buffer->GetColumns(); ++Col)
+			{
+				Buffer->SetCell(CursorAbsoluteRow, Col, BlankCell);
+			}
+		}
+		break;
+	}
+	case 'S': // SU  - scroll up
+	{
+		const int32 Count = FMath::Max(1, GetCSIParam(0, 1));
+		Buffer->ScrollRegionUp(GetEffectiveScrollTop(), GetEffectiveScrollBottom(), Count);
+		break;
+	}
+	case 'T': // SD  - scroll down
+	{
+		const int32 Count = FMath::Max(1, GetCSIParam(0, 1));
+		Buffer->ScrollRegionDown(GetEffectiveScrollTop(), GetEffectiveScrollBottom(), Count);
+		break;
+	}
+	case 'r': // DECSTBM  - set scroll region
+	{
+		const int32 Top = FMath::Max(1, GetCSIParam(0, 1)) - 1;
+		const int32 Bottom = FMath::Max(1, GetCSIParam(1, Buffer->GetViewportRows())) - 1;
+		if (Top < Bottom && Bottom < Buffer->GetViewportRows())
+		{
+			ScrollRegionTop = Top;
+			ScrollRegionBottom = Bottom;
+		}
+		else
+		{
+			ScrollRegionTop = -1;
+			ScrollRegionBottom = -1;
+		}
+		// DECSTBM also moves cursor to home.
+		Buffer->Cursor.Row = 0;
+		Buffer->Cursor.Column = 0;
+		break;
+	}
+	case 'm': // SGR  - select graphic rendition
+		ExecuteSGR();
+		break;
+	case 'G': // CHA  - cursor horizontal absolute
+	case '`': // HPA  - same
+	{
+		const int32 Col = FMath::Max(1, GetCSIParam(0, 1)) - 1;
+		Buffer->Cursor.Column = FMath::Clamp(Col, 0, Buffer->GetColumns() - 1);
+		break;
+	}
+	case 'd': // VPA  - vertical position absolute
+	{
+		const int32 Row = FMath::Max(1, GetCSIParam(0, 1)) - 1;
+		Buffer->Cursor.Row = FMath::Clamp(Row, 0, Buffer->GetViewportRows() - 1);
+		break;
+	}
+	case 'L': // IL  - insert lines
+	{
+		const int32 Count = FMath::Max(1, GetCSIParam(0, 1));
+		Buffer->ScrollRegionDown(Buffer->Cursor.Row, GetEffectiveScrollBottom(), Count);
+		break;
+	}
+	case 'M': // DL  - delete lines
+	{
+		const int32 Count = FMath::Max(1, GetCSIParam(0, 1));
+		Buffer->ScrollRegionUp(Buffer->Cursor.Row, GetEffectiveScrollBottom(), Count);
+		break;
+	}
+	case 'P': // DCH  - delete characters
+	{
+		const int32 Count = FMath::Max(1, GetCSIParam(0, 1));
+		const int32 AbsRow = Buffer->GetAbsoluteRow(Buffer->Cursor.Row, 0);
+		const int32 Columns = Buffer->GetColumns();
+		for (int32 Col = Buffer->Cursor.Column; Col < Columns; ++Col)
+		{
+			const int32 SourceCol = Col + Count;
+			if (SourceCol < Columns)
+			{
+				Buffer->SetCell(AbsRow, Col, Buffer->GetCell(AbsRow, SourceCol));
+			}
+			else
+			{
+				Buffer->SetCell(AbsRow, Col, Buffer->DefaultCell);
+			}
+		}
+		break;
+	}
+	case '@': // ICH  - insert characters
+	{
+		const int32 Count = FMath::Max(1, GetCSIParam(0, 1));
+		const int32 AbsRow = Buffer->GetAbsoluteRow(Buffer->Cursor.Row, 0);
+		const int32 Columns = Buffer->GetColumns();
+		for (int32 Col = Columns - 1; Col >= Buffer->Cursor.Column; --Col)
+		{
+			const int32 SourceCol = Col - Count;
+			if (SourceCol >= Buffer->Cursor.Column)
+			{
+				Buffer->SetCell(AbsRow, Col, Buffer->GetCell(AbsRow, SourceCol));
+			}
+			else
+			{
+				Buffer->SetCell(AbsRow, Col, Buffer->DefaultCell);
+			}
+		}
+		break;
+	}
+	case 'X': // ECH  - erase characters
+	{
+		const int32 Count = FMath::Max(1, GetCSIParam(0, 1));
+		const int32 AbsRow = Buffer->GetAbsoluteRow(Buffer->Cursor.Row, 0);
+		FTerminalCell BlankCell = Buffer->DefaultCell;
+		BlankCell.Foreground = CurrentForeground;
+		BlankCell.Background = CurrentBackground;
+		for (int32 Index = 0; Index < Count && (Buffer->Cursor.Column + Index) < Buffer->GetColumns(); ++Index)
+		{
+			Buffer->SetCell(AbsRow, Buffer->Cursor.Column + Index, BlankCell);
+		}
+		break;
+	}
+	case 's': // SCOSC  - save cursor position
+		Buffer->Cursor.SavedRow = Buffer->Cursor.Row;
+		Buffer->Cursor.SavedColumn = Buffer->Cursor.Column;
+		break;
+	case 'u': // SCORC  - restore cursor position
+		Buffer->Cursor.Row = Buffer->Cursor.SavedRow;
+		Buffer->Cursor.Column = Buffer->Cursor.SavedColumn;
+		break;
+	case 'g': // TBC  - tab clear
+	{
+		const int32 Mode = GetCSIParam(0, 0);
+		if (Mode == 0)
+		{
+			if (Buffer->Cursor.Column < TabStops.Num())
+			{
+				TabStops[Buffer->Cursor.Column] = false;
+			}
+		}
+		else if (Mode == 3)
+		{
+			for (int32 Index = 0; Index < TabStops.Num(); ++Index)
+			{
+				TabStops[Index] = false;
+			}
+		}
+		break;
+	}
+	case 'n': // DSR  - device status report
+	{
+		const int32 Mode = GetCSIParam(0, 0);
+		if (Mode == 5)
+		{
+			// Operating status report: respond "OK".
+			static const uint8 StatusOKResponse[] = { 0x1B, '[', '0', 'n' };
+			ResponseBuffer.Append(StatusOKResponse, UE_ARRAY_COUNT(StatusOKResponse));
+		}
+		else if (Mode == 6 && Buffer)
+		{
+			// Cursor position report: ESC [ row ; col R  (1-based).
+			const FString Report = FString::Printf(TEXT("\x1B[%d;%dR"), Buffer->Cursor.Row + 1, Buffer->Cursor.Column + 1);
+			FTCHARToUTF8 ReportUtf8(*Report);
+			ResponseBuffer.Append(reinterpret_cast<const uint8*>(ReportUtf8.Get()), ReportUtf8.Length());
+		}
+		break;
+	}
+	case 'c': // DA  - device attributes
+	{
+		// Respond as VT101 with no options.
+		static const uint8 DeviceAttributesResponse[] = { 0x1B, '[', '?', '1', ';', '0', 'c' };
+		ResponseBuffer.Append(DeviceAttributesResponse, UE_ARRAY_COUNT(DeviceAttributesResponse));
+		break;
+	}
+	case 't': // XTWINOPS  - window manipulation (silently ignore)
+		break;
+	case 'b': // REP  - repeat previous character (silently ignore for now)
+		break;
+	default:
+		// Unrecognized CSI  - silently ignore.
+		break;
+	}
+}
+
+// ── SGR (Select Graphic Rendition) ──
+
+void FVTParser::ExecuteSGR()
+{
+	if (CSIParams.Num() == 0)
+	{
+		CSIParams.Add(0);
+	}
+
+	for (int32 Index = 0; Index < CSIParams.Num(); ++Index)
+	{
+		const int32 Param = CSIParams[Index];
+
+		switch (Param)
+		{
+		case 0: // Reset
+			CurrentForeground = Buffer ? Buffer->DefaultCell.Foreground : FColor(212, 212, 212);
+			CurrentBackground = Buffer ? Buffer->DefaultCell.Background : FColor(30, 30, 30);
+			CurrentAttributes = ETerminalAttribute::None;
+			break;
+		case 1: CurrentAttributes |= ETerminalAttribute::Bold; break;
+		case 2: CurrentAttributes |= ETerminalAttribute::Dim; break;
+		case 3: CurrentAttributes |= ETerminalAttribute::Italic; break;
+		case 4: CurrentAttributes |= ETerminalAttribute::Underline; break;
+		case 7: CurrentAttributes |= ETerminalAttribute::Inverse; break;
+		case 9: CurrentAttributes |= ETerminalAttribute::Strikethrough; break;
+		case 22: CurrentAttributes &= ~(ETerminalAttribute::Bold | ETerminalAttribute::Dim); break;
+		case 23: CurrentAttributes &= ~ETerminalAttribute::Italic; break;
+		case 24: CurrentAttributes &= ~ETerminalAttribute::Underline; break;
+		case 27: CurrentAttributes &= ~ETerminalAttribute::Inverse; break;
+		case 29: CurrentAttributes &= ~ETerminalAttribute::Strikethrough; break;
+
+		// Standard foreground colors (30-37).
+		case 30: case 31: case 32: case 33:
+		case 34: case 35: case 36: case 37:
+			CurrentForeground = AnsiPalette[Param - 30];
+			break;
+
+		// Extended foreground (38;5;N or 38;2;R;G;B).
+		case 38:
+			if (Index + 1 < CSIParams.Num())
+			{
+				if (CSIParams[Index + 1] == 5 && Index + 2 < CSIParams.Num())
+				{
+					// 256-color
+					const int32 ColorIndex = CSIParams[Index + 2];
+					CurrentForeground = Get256Color(ColorIndex);
+					Index += 2;
+				}
+				else if (CSIParams[Index + 1] == 2 && Index + 4 < CSIParams.Num())
+				{
+					// True-color
+					CurrentForeground = FColor(
+						FMath::Clamp(CSIParams[Index + 2], 0, 255),
+						FMath::Clamp(CSIParams[Index + 3], 0, 255),
+						FMath::Clamp(CSIParams[Index + 4], 0, 255));
+					Index += 4;
+				}
+			}
+			break;
+
+		case 39: // Default foreground
+			CurrentForeground = Buffer ? Buffer->DefaultCell.Foreground : FColor(212, 212, 212);
+			break;
+
+		// Standard background colors (40-47).
+		case 40: case 41: case 42: case 43:
+		case 44: case 45: case 46: case 47:
+			CurrentBackground = AnsiPalette[Param - 40];
+			break;
+
+		// Extended background (48;5;N or 48;2;R;G;B).
+		case 48:
+			if (Index + 1 < CSIParams.Num())
+			{
+				if (CSIParams[Index + 1] == 5 && Index + 2 < CSIParams.Num())
+				{
+					const int32 ColorIndex = CSIParams[Index + 2];
+					CurrentBackground = Get256Color(ColorIndex);
+					Index += 2;
+				}
+				else if (CSIParams[Index + 1] == 2 && Index + 4 < CSIParams.Num())
+				{
+					CurrentBackground = FColor(
+						FMath::Clamp(CSIParams[Index + 2], 0, 255),
+						FMath::Clamp(CSIParams[Index + 3], 0, 255),
+						FMath::Clamp(CSIParams[Index + 4], 0, 255));
+					Index += 4;
+				}
+			}
+			break;
+
+		case 49: // Default background
+			CurrentBackground = Buffer ? Buffer->DefaultCell.Background : FColor(30, 30, 30);
+			break;
+
+		// Bright foreground colors (90-97).
+		case 90: case 91: case 92: case 93:
+		case 94: case 95: case 96: case 97:
+			CurrentForeground = AnsiPalette[Param - 90 + 8];
+			break;
+
+		// Bright background colors (100-107).
+		case 100: case 101: case 102: case 103:
+		case 104: case 105: case 106: case 107:
+			CurrentBackground = AnsiPalette[Param - 100 + 8];
+			break;
+
+		default:
+			break;
+		}
+	}
+}
+
+// ── Private Mode (CSI ? / CSI > / etc.) ──
+
+void FVTParser::ExecutePrivateMode(uint32 FinalByte)
+{
+	// Only handle '?' prefix for DECSET/DECRST. Other prefixes ('<', '=', '>')
+	// are silently ignored (e.g., Kitty keyboard protocol pushes via CSI > u).
+	if (CSIPrivatePrefix != '?')
+	{
+		return;
+	}
+
+	for (int32 Index = 0; Index < CSIParams.Num(); ++Index)
+	{
+		const int32 Mode = CSIParams[Index];
+
+		switch (FinalByte)
+		{
+		case 'h': // DECSET  - set mode
+			switch (Mode)
+			{
+			case 1: bApplicationCursorKeys = true; break;
+			case 7: bAutoWrap = true; break;
+			case 25: bCursorVisible = true; break;
+			case 1000: MouseTrackingMode = EMouseTrackingMode::Normal; break;
+			case 1002: MouseTrackingMode = EMouseTrackingMode::ButtonEvent; break;
+			case 1003: MouseTrackingMode = EMouseTrackingMode::Any; break;
+			case 1006: bSGRMouseEncoding = true; break;
+			case 1049: // Alternate screen buffer
+				Buffer->ActivateAlternateBuffer();
+				break;
+			case 2004: bBracketedPaste = true; break;
+			case 2026: bSynchronizedOutput = true; break;
+			}
+			break;
+		case 'l': // DECRST  - reset mode
+			switch (Mode)
+			{
+			case 1: bApplicationCursorKeys = false; break;
+			case 7: bAutoWrap = false; break;
+			case 25: bCursorVisible = false; break;
+			case 1000: // Falls through - any mouse tracking reset disables tracking.
+			case 1002:
+			case 1003: MouseTrackingMode = EMouseTrackingMode::None; break;
+			case 1006: bSGRMouseEncoding = false; break;
+			case 1049:
+				Buffer->DeactivateAlternateBuffer();
+				break;
+			case 2004: bBracketedPaste = false; break;
+			case 2026: bSynchronizedOutput = false; break;
+			}
+			break;
+		}
+	}
+}
+
+// ── OSC Dispatch ──
+
+void FVTParser::ExecuteOSC()
+{
+	// Parse OSC payload: "Ps;Pt" where Ps is the command number.
+	int32 SemicolonIndex;
+	if (OSCPayload.FindChar(TEXT(';'), SemicolonIndex))
+	{
+		const FString CommandStr = OSCPayload.Left(SemicolonIndex);
+		const int32 Command = FCString::Atoi(*CommandStr);
+		const FString Payload = OSCPayload.Mid(SemicolonIndex + 1);
+
+		switch (Command)
+		{
+		case 0: // Set window title and icon name
+		case 2: // Set window title
+			WindowTitle = Payload;
+			break;
+		default:
+			// Ignore other OSC commands.
+			break;
+		}
+	}
+}
+
+// ── Character Output ──
+
+bool FVTParser::IsZeroWidthCodepoint(uint32 Codepoint)
+{
+	// Variation Selectors (emoji presentation, text presentation, etc.)
+	if (Codepoint >= 0xFE00 && Codepoint <= 0xFE0F) return true;
+	// Combining Diacritical Marks
+	if (Codepoint >= 0x0300 && Codepoint <= 0x036F) return true;
+	// Zero-width space, non-joiner, joiner
+	if (Codepoint >= 0x200B && Codepoint <= 0x200D) return true;
+	// Zero-width no-break space (BOM)
+	if (Codepoint == 0xFEFF) return true;
+	// Combining Diacritical Marks Extended
+	if (Codepoint >= 0x1AB0 && Codepoint <= 0x1AFF) return true;
+	// Combining Diacritical Marks Supplement
+	if (Codepoint >= 0x1DC0 && Codepoint <= 0x1DFF) return true;
+	// Combining Diacritical Marks for Symbols
+	if (Codepoint >= 0x20D0 && Codepoint <= 0x20FF) return true;
+	// Combining Half Marks
+	if (Codepoint >= 0xFE20 && Codepoint <= 0xFE2F) return true;
+	// Variation Selectors Supplement
+	if (Codepoint >= 0xE0100 && Codepoint <= 0xE01EF) return true;
+	return false;
+}
+
+void FVTParser::PutCharacter(TCHAR Character)
+{
+	if (!Buffer)
+	{
+		return;
+	}
+
+	const int32 Columns = Buffer->GetColumns();
+
+	// Auto-wrap: if cursor is past the last column, wrap to next line.
+	if (Buffer->Cursor.Column >= Columns)
+	{
+		if (bAutoWrap)
+		{
+			Buffer->Cursor.Column = 0;
+			LineFeed();
+		}
+		else
+		{
+			Buffer->Cursor.Column = Columns - 1;
+		}
+	}
+
+	const int32 AbsoluteRow = Buffer->GetAbsoluteRow(Buffer->Cursor.Row, 0);
+
+	FTerminalCell Cell;
+	Cell.Character = Character;
+	Cell.Foreground = CurrentForeground;
+	Cell.Background = CurrentBackground;
+	Cell.Attributes = CurrentAttributes;
+
+	Buffer->SetCell(AbsoluteRow, Buffer->Cursor.Column, Cell);
+	Buffer->Cursor.Column++;
+}
+
+void FVTParser::LineFeed()
+{
+	if (!Buffer)
+	{
+		return;
+	}
+
+	const int32 Bottom = GetEffectiveScrollBottom();
+
+	if (Buffer->Cursor.Row == Bottom)
+	{
+		// At the bottom of the scroll region  - scroll up.
+		const int32 Top = GetEffectiveScrollTop();
+		if (Top == 0 && Bottom == Buffer->GetViewportRows() - 1)
+		{
+			// Full-screen scroll  - push a new line into the ring buffer (preserves scrollback).
+			Buffer->PushNewLine();
+		}
+		else
+		{
+			// Scroll within a region.
+			Buffer->ScrollRegionUp(Top, Bottom, 1);
+		}
+	}
+	else if (Buffer->Cursor.Row < Buffer->GetViewportRows() - 1)
+	{
+		Buffer->Cursor.Row++;
+	}
+}
+
+int32 FVTParser::GetEffectiveScrollTop() const
+{
+	return (ScrollRegionTop >= 0) ? ScrollRegionTop : 0;
+}
+
+int32 FVTParser::GetEffectiveScrollBottom() const
+{
+	if (!Buffer)
+	{
+		return 0;
+	}
+	return (ScrollRegionBottom >= 0) ? ScrollRegionBottom : (Buffer->GetViewportRows() - 1);
+}
+
+void FVTParser::InitializeTabStops()
+{
+	const int32 MaxColumns = Buffer ? Buffer->GetColumns() : 256;
+	TabStops.SetNum(MaxColumns);
+	for (int32 Index = 0; Index < MaxColumns; ++Index)
+	{
+		TabStops[Index] = (Index % 8 == 0);
+	}
+}
+
+int32 FVTParser::GetCSIParam(int32 Index, int32 Default) const
+{
+	if (Index < CSIParams.Num() && CSIParams[Index] != 0)
+	{
+		return CSIParams[Index];
+	}
+	return Default;
+}
+
+FColor FVTParser::Get256Color(int32 Index) const
+{
+	// Standard 16 colors (0-15)  - use the configurable ANSI palette.
+	if (Index >= 0 && Index < 16)
+	{
+		return AnsiPalette[Index];
+	}
+
+	// 216-color cube (16-231).
+	if (Index >= 16 && Index < 232)
+	{
+		const int32 CubeIndex = Index - 16;
+		const int32 Red = CubeIndex / 36;
+		const int32 Green = (CubeIndex / 6) % 6;
+		const int32 Blue = CubeIndex % 6;
+		return FColor(
+			Red ? (Red * 40 + 55) : 0,
+			Green ? (Green * 40 + 55) : 0,
+			Blue ? (Blue * 40 + 55) : 0);
+	}
+
+	// Grayscale ramp (232-255).
+	if (Index >= 232 && Index < 256)
+	{
+		const int32 Gray = (Index - 232) * 10 + 8;
+		return FColor(Gray, Gray, Gray);
+	}
+
+	return FColor(255, 255, 255);
+}
diff --git a/Engine/Plugins/Experimental/Terminal/Source/Terminal/Public/ConPTYSession.h b/Engine/Plugins/Experimental/Terminal/Source/Terminal/Public/ConPTYSession.h
new file mode 100644
--- /dev/null
+++ b/Engine/Plugins/Experimental/Terminal/Source/Terminal/Public/ConPTYSession.h
@@ -0,0 +1,71 @@
+// Copyright Epic Games, Inc. All Rights Reserved.
+
+#pragma once
+
+#include "ITerminalSession.h"
+
+#if PLATFORM_WINDOWS
+
+#include "HAL/Runnable.h"
+#include "HAL/RunnableThread.h"
+#include "HAL/CriticalSection.h"
+
+/**
+ * Manages a ConPTY pseudo-console session on Windows.
+ *
+ * Creates a pseudo-console via the ConPTY API (dynamically loaded),
+ * spawns a shell process attached to it, and runs asynchronous I/O
+ * on a background thread with staged dispatch to the game thread.
+ */
+class FConPTYSession : public ITerminalSession
+{
+public:
+
+	FConPTYSession();
+	virtual ~FConPTYSession() override;
+
+	FConPTYSession(const FConPTYSession&) = delete;
+	FConPTYSession& operator=(const FConPTYSession&) = delete;
+
+	/** Check if the ConPTY API is available on this system. */
+	static bool IsConPTYAvailable();
+
+	virtual bool Create(const FString& ShellPath, const FString& WorkingDirectory, int32 Columns, int32 Rows) override;
+	virtual void Shutdown() override;
+	virtual void WriteInput(TArrayView<const uint8> Data) override;
+	virtual void Resize(int32 Columns, int32 Rows) override;
+	virtual TArray<uint8> ConsumeOutput() override;
+	virtual bool IsRunning() const override;
+
+private:
+
+	/** Background I/O reader thread. */
+	class FReaderRunnable : public FRunnable
+	{
+	public:
+		FReaderRunnable(FConPTYSession* InOwner) : Owner(InOwner) {}
+		virtual uint32 Run() override;
+		virtual void Stop() override;
+	private:
+		FConPTYSession* Owner = nullptr;
+		TAtomic<bool> bStopRequested{false};
+	};
+
+	void* PseudoConsoleHandle = nullptr;
+	void* InputPipeWrite = nullptr;
+	void* OutputPipeRead = nullptr;
+	void* ProcessHandle = nullptr;
+	void* ThreadHandle = nullptr;
+
+	/** Shared alive flag for safe async callback dispatch after destruction. */
+	TSharedPtr<TAtomic<bool>> bAlive;
+
+	/** Output staging buffer  - filled by the reader thread, drained by ConsumeOutput() on the game thread. */
+	TArray<uint8> StagingBuffer;
+	FCriticalSection StagingLock;
+
+	TUniquePtr<FReaderRunnable> ReaderRunnable;
+	TUniquePtr<FRunnableThread> ReaderThread;
+};
+
+#endif // PLATFORM_WINDOWS
diff --git a/Engine/Plugins/Experimental/Terminal/Source/Terminal/Public/ITerminalSession.h b/Engine/Plugins/Experimental/Terminal/Source/Terminal/Public/ITerminalSession.h
new file mode 100644
--- /dev/null
+++ b/Engine/Plugins/Experimental/Terminal/Source/Terminal/Public/ITerminalSession.h
@@ -0,0 +1,62 @@
+// Copyright Epic Games, Inc. All Rights Reserved.
+
+#pragma once
+
+#include "CoreMinimal.h"
+
+DECLARE_LOG_CATEGORY_EXTERN(LogTerminal, Log, All);
+
+DECLARE_DELEGATE_OneParam(FOnTerminalProcessExited, int32 /* ExitCode */);
+
+/**
+ * Abstract interface for a terminal PTY session.
+ *
+ * Platform-specific backends (FConPTYSession on Windows, FPosixPTYSession on
+ * Linux/macOS) implement this interface so that the terminal widget can drive
+ * any backend without platform conditionals in its core logic.
+ */
+class ITerminalSession
+{
+public:
+
+	virtual ~ITerminalSession() = default;
+
+	/**
+	 * Create a platform-appropriate session backend.
+	 *
+	 * Returns nullptr with an error message if no backend is available (e.g.
+	 * ConPTY missing on older Windows, or unsupported platform).
+	 *
+	 * @param OutError Receives a user-visible error string when nullptr is returned.
+	 */
+	static TSharedPtr<ITerminalSession> CreateForCurrentPlatform(FString& OutError);
+
+	/**
+	 * Create a new session.
+	 *
+	 * @param ShellPath Path to the shell executable. Empty uses the system default.
+	 * @param WorkingDirectory Initial working directory for the shell.
+	 * @param Columns Initial column count.
+	 * @param Rows Initial row count.
+	 * @return true if the session was created successfully.
+	 */
+	virtual bool Create(const FString& ShellPath, const FString& WorkingDirectory, int32 Columns, int32 Rows) = 0;
+
+	/** Shut down the session, closing the PTY and terminating the shell process. */
+	virtual void Shutdown() = 0;
+
+	/** Write input bytes to the PTY input pipe. */
+	virtual void WriteInput(TArrayView<const uint8> Data) = 0;
+
+	/** Resize the PTY. */
+	virtual void Resize(int32 Columns, int32 Rows) = 0;
+
+	/** Swap the output staging buffer and return the accumulated bytes for parsing on the game thread. */
+	virtual TArray<uint8> ConsumeOutput() = 0;
+
+	/** Check if the shell process is still running. */
+	virtual bool IsRunning() const = 0;
+
+	/** Delegate fired when the shell process exits. */
+	FOnTerminalProcessExited OnProcessExited;
+};
diff --git a/Engine/Plugins/Experimental/Terminal/Source/Terminal/Public/PosixPTYSession.h b/Engine/Plugins/Experimental/Terminal/Source/Terminal/Public/PosixPTYSession.h
new file mode 100644
--- /dev/null
+++ b/Engine/Plugins/Experimental/Terminal/Source/Terminal/Public/PosixPTYSession.h
@@ -0,0 +1,68 @@
+// Copyright Epic Games, Inc. All Rights Reserved.
+
+#pragma once
+
+#include "ITerminalSession.h"
+
+#if PLATFORM_UNIX || PLATFORM_MAC
+
+#include "HAL/Runnable.h"
+#include "HAL/RunnableThread.h"
+#include "HAL/CriticalSection.h"
+
+/**
+ * Manages a POSIX PTY session on Linux and macOS.
+ *
+ * Creates a pseudo-terminal via forkpty(), spawns a shell process, and
+ * runs asynchronous I/O on a background thread with staged dispatch to
+ * the game thread  - mirroring the FConPTYSession pattern on Windows.
+ */
+class FPosixPTYSession : public ITerminalSession
+{
+public:
+
+	FPosixPTYSession();
+	virtual ~FPosixPTYSession() override;
+
+	FPosixPTYSession(const FPosixPTYSession&) = delete;
+	FPosixPTYSession& operator=(const FPosixPTYSession&) = delete;
+
+	virtual bool Create(const FString& ShellPath, const FString& WorkingDirectory, int32 Columns, int32 Rows) override;
+	virtual void Shutdown() override;
+	virtual void WriteInput(TArrayView<const uint8> Data) override;
+	virtual void Resize(int32 Columns, int32 Rows) override;
+	virtual TArray<uint8> ConsumeOutput() override;
+	virtual bool IsRunning() const override;
+
+private:
+
+	/** Background I/O reader thread. */
+	class FReaderRunnable : public FRunnable
+	{
+	public:
+		FReaderRunnable(FPosixPTYSession* InOwner) : Owner(InOwner) {}
+		virtual uint32 Run() override;
+		virtual void Stop() override;
+	private:
+		FPosixPTYSession* Owner = nullptr;
+		TAtomic<bool> bStopRequested{false};
+	};
+
+	/** Master file descriptor of the pseudo-terminal. */
+	int MasterFileDescriptor = -1;
+
+	/** PID of the child shell process. */
+	pid_t ProcessId = -1;
+
+	/** Shared alive flag for safe async callback dispatch after destruction. */
+	TSharedPtr<TAtomic<bool>> bAlive;
+
+	/** Output staging buffer  - filled by the reader thread, drained by ConsumeOutput() on the game thread. */
+	TArray<uint8> StagingBuffer;
+	FCriticalSection StagingLock;
+
+	TUniquePtr<FReaderRunnable> ReaderRunnable;
+	TUniquePtr<FRunnableThread> ReaderThread;
+};
+
+#endif // PLATFORM_UNIX || PLATFORM_MAC
diff --git a/Engine/Plugins/Experimental/Terminal/Source/Terminal/Public/STerminal.h b/Engine/Plugins/Experimental/Terminal/Source/Terminal/Public/STerminal.h
new file mode 100644
--- /dev/null
+++ b/Engine/Plugins/Experimental/Terminal/Source/Terminal/Public/STerminal.h
@@ -0,0 +1,165 @@
+// Copyright Epic Games, Inc. All Rights Reserved.
+
+#pragma once
+
+#include "CoreMinimal.h"
+#include "Fonts/CompositeFont.h"
+#include "Widgets/SLeafWidget.h"
+#include "TerminalBuffer.h"
+#include "VTParser.h"
+
+class ITerminalSession;
+class SScrollBar;
+struct FTerminalColorScheme;
+
+/**
+ * Native Slate terminal emulator widget.
+ *
+ * Renders a character cell grid directly via OnPaint using FSlateDrawElement.
+ * Backed by FTerminalBuffer + FVTParser + FConPTYSession.
+ */
+class STerminal : public SLeafWidget
+{
+	friend class FSTerminalSpec;
+
+public:
+
+	SLATE_BEGIN_ARGS(STerminal)
+	{}
+		SLATE_ARGUMENT(TSharedPtr<SScrollBar>, ExternalScrollbar)
+	SLATE_END_ARGS()
+
+	TERMINAL_API void Construct(const FArguments& InArgs);
+
+	TERMINAL_API virtual ~STerminal() override;
+
+	/** Write a command string to the terminal session followed by a carriage return. */
+	TERMINAL_API void ExecuteCommand(const FString& Command);
+
+	/** Returns true once the terminal session has been created and is running. */
+	TERMINAL_API bool IsSessionRunning() const;
+
+	//~ Begin SWidget interface
+	TERMINAL_API virtual FVector2D ComputeDesiredSize(float LayoutScaleMultiplier) const override;
+	TERMINAL_API virtual int32 OnPaint(const FPaintArgs& Args, const FGeometry& AllottedGeometry, const FSlateRect& MyCullingRect, FSlateWindowElementList& OutDrawElements, int32 LayerId, const FWidgetStyle& InWidgetStyle, bool bParentEnabled) const override;
+	TERMINAL_API virtual FReply OnKeyChar(const FGeometry& MyGeometry, const FCharacterEvent& InCharacterEvent) override;
+	TERMINAL_API virtual FReply OnKeyDown(const FGeometry& MyGeometry, const FKeyEvent& InKeyEvent) override;
+	TERMINAL_API virtual FReply OnMouseButtonDown(const FGeometry& MyGeometry, const FPointerEvent& MouseEvent) override;
+	TERMINAL_API virtual FReply OnMouseButtonUp(const FGeometry& MyGeometry, const FPointerEvent& MouseEvent) override;
+	TERMINAL_API virtual FReply OnMouseMove(const FGeometry& MyGeometry, const FPointerEvent& MouseEvent) override;
+	TERMINAL_API virtual FReply OnMouseWheel(const FGeometry& MyGeometry, const FPointerEvent& MouseEvent) override;
+	TERMINAL_API virtual FReply OnMouseButtonDoubleClick(const FGeometry& MyGeometry, const FPointerEvent& MouseEvent) override;
+	virtual bool SupportsKeyboardFocus() const override { return true; }
+	TERMINAL_API virtual FReply OnFocusReceived(const FGeometry& MyGeometry, const FFocusEvent& InFocusEvent) override;
+	TERMINAL_API virtual void OnFocusLost(const FFocusEvent& InFocusEvent) override;
+	TERMINAL_API virtual FCursorReply OnCursorQuery(const FGeometry& MyGeometry, const FPointerEvent& CursorEvent) const override;
+	//~ End SWidget interface
+
+private:
+
+	/** Tick handler  - consumes session output and parses it. */
+	EActiveTimerReturnType OnTick(double InCurrentTime, float InDeltaTime);
+
+	/** Cursor blink timer. */
+	EActiveTimerReturnType OnCursorBlink(double InCurrentTime, float InDeltaTime);
+
+	/** Measure the monospace cell dimensions from the current font at the given layout scale. */
+	void MeasureCellSize(float LayoutScale = 1.0f);
+
+	/** Deferred initialization  - creates buffer and terminal session at correct dimensions. */
+	void InitializeTerminal();
+
+	/** Recompute viewport rows and columns from the widget's allocated geometry. */
+	void UpdateViewportDimensions(const FGeometry& AllottedGeometry);
+
+	/** Update the external scrollbar thumb position and size to reflect current scroll state. */
+	void UpdateScrollBar();
+
+	/** Convert a local pixel position to a buffer-absolute cell coordinate. */
+	TERMINAL_API FIntPoint PixelToCell(const FGeometry& MyGeometry, const FVector2D& LocalPosition) const;
+
+	/** Check if a cell is within the current selection range. */
+	TERMINAL_API bool IsCellSelected(int32 AbsoluteRow, int32 Column) const;
+
+	/** Get the selected text as a string. */
+	TERMINAL_API FString GetSelectedText() const;
+
+	/** Clear the current selection. */
+	TERMINAL_API void ClearSelection();
+
+	/** Paint per-cell backgrounds (inverse-aware). */
+	void PaintBackground(const FGeometry& AllottedGeometry, FSlateWindowElementList& OutDrawElements, int32 LayerId) const;
+
+	/** Paint text with monospace-safe batching, underline, and strikethrough. */
+	void PaintText(const FGeometry& AllottedGeometry, FSlateWindowElementList& OutDrawElements, int32 LayerId) const;
+
+	/** Paint foreground overlays (selection highlight, cursor). */
+	void PaintForeground(const FGeometry& AllottedGeometry, FSlateWindowElementList& OutDrawElements, int32 LayerId) const;
+
+	/** Compute the effective foreground color for a cell, applying inverse, bold, and dim. */
+	static FColor ComputeForegroundColor(const FTerminalCell& Cell);
+
+	/** Apply a color scheme to default cell colors. */
+	void ApplyColorScheme(const FTerminalColorScheme& Scheme);
+
+	/** Called when the shell process exits. */
+	void HandleProcessExited(int32 ExitCode);
+
+	/** Called when the user drags the external scrollbar. */
+	void HandleScrollBarScrolled(float InScrollOffsetFraction);
+
+	/** Send a mouse event escape sequence to the session when mouse tracking is active. */
+	void SendMouseEvent(int32 Button, int32 Column, int32 Row, bool bRelease);
+
+	/** Core terminal components. */
+	FTerminalBuffer Buffer;
+	FVTParser Parser;
+	TSharedPtr<ITerminalSession> Session;
+	TSharedPtr<SScrollBar> ScrollBar;
+
+	/** Font and rendering. */
+	TSharedPtr<FStandaloneCompositeFont> TerminalFont;
+	FSlateFontInfo FontInfo;
+	float CellWidth = 8.0f;
+	float CellHeight = 16.0f;
+	FColor SelectionHighlightColor = FColor(38, 79, 120);
+	FColor CursorDisplayColor = FColor(174, 175, 173, 180);
+
+	/** Viewport dimensions. */
+	int32 ViewportColumns = 80;
+	int32 ViewportRows = 24;
+	int32 ScrollOffset = 0;
+	mutable FVector2D LastAllocatedSize = FVector2D::ZeroVector;
+
+	/** Per-tick parse budget in bytes. */
+	static constexpr int32 ParseBudgetPerTick = 256 * 1024;
+
+	/** Leftover bytes from a previous tick that exceeded the parse budget. */
+	TArray<uint8> PendingParseData;
+
+	/** Selection state (buffer-absolute coordinates). */
+	FIntPoint SelectionAnchor = FIntPoint(-1, -1);
+	FIntPoint SelectionActive = FIntPoint(-1, -1);
+
+	/** Auto-scroll during drag selection: positive = scroll up (older), negative = scroll down (newer), 0 = none. */
+	int32 SelectionAutoScrollDirection = 0;
+	float SelectionAutoScrollAccumulator = 0.0f;
+
+	/** Click tracking for double/triple click. */
+	int32 ClickCount = 0;
+	double LastClickTime = 0.0;
+	FIntPoint LastClickCell = FIntPoint(-1, -1);
+
+	/** Diagnostic counters. */
+	int64 TotalBytesConsumed = 0;
+	double LastOutputTime = 0.0;
+	double LastHeartbeatLogTime = 0.0;
+
+	/** State flags. */
+	bool bInitialized = false;
+	bool bSessionDead = false;
+	bool bReceivedFirstOutput = false;
+	bool bHasFocus = false;
+	bool bSelecting = false;
+	bool bHasSelection = false;
+};
diff --git a/Engine/Plugins/Experimental/Terminal/Source/Terminal/Public/TerminalBuffer.h b/Engine/Plugins/Experimental/Terminal/Source/Terminal/Public/TerminalBuffer.h
new file mode 100644
--- /dev/null
+++ b/Engine/Plugins/Experimental/Terminal/Source/Terminal/Public/TerminalBuffer.h
@@ -0,0 +1,132 @@
+// Copyright Epic Games, Inc. All Rights Reserved.
+
+#pragma once
+
+#include "CoreMinimal.h"
+
+/** Bitfield for terminal cell style attributes. */
+namespace ETerminalAttribute
+{
+	enum Type : uint8
+	{
+		None          = 0,
+		Bold          = 1 << 0,
+		Dim           = 1 << 1,
+		Italic        = 1 << 2,
+		Underline     = 1 << 3,
+		Inverse       = 1 << 4,
+		Strikethrough = 1 << 5,
+	};
+}
+
+/** A single character cell in the terminal grid. */
+struct FTerminalCell
+{
+	TCHAR Character = TEXT(' ');
+	FColor Foreground = FColor(212, 212, 212);
+	FColor Background = FColor(30, 30, 30);
+	uint8 Attributes = ETerminalAttribute::None;
+};
+
+/** Cursor state for the terminal. */
+struct FTerminalCursor
+{
+	int32 Row = 0;
+	int32 Column = 0;
+	bool bBlinkOn = true;
+
+	/** Saved cursor position for DECSC/DECRC. */
+	int32 SavedRow = 0;
+	int32 SavedColumn = 0;
+};
+
+/**
+ * Ring-buffer-backed character cell grid with configurable scrollback.
+ *
+ * Rows are stored in a flat ring buffer. The viewport (the visible portion)
+ * is always the last ViewportRows rows of the active buffer. Scrollback rows
+ * accumulate above the viewport up to ScrollbackLimit.
+ */
+class FTerminalBuffer
+{
+public:
+
+	TERMINAL_API FTerminalBuffer();
+
+	TERMINAL_API void Initialize(int32 InColumns, int32 InViewportRows, int32 InScrollbackLimit);
+
+	TERMINAL_API void Resize(int32 InColumns, int32 InViewportRows);
+
+	/** Get a cell by absolute row index (0 = oldest scrollback row) and column. */
+	TERMINAL_API const FTerminalCell& GetCell(int32 AbsoluteRow, int32 Column) const;
+
+	/** Set a cell by absolute row index and column. */
+	TERMINAL_API void SetCell(int32 AbsoluteRow, int32 Column, const FTerminalCell& Cell);
+
+	/** Push a new blank line at the bottom, advancing the ring buffer. */
+	TERMINAL_API void PushNewLine();
+
+	/** Get the absolute row index for a given viewport row (0 = top of viewport) and scroll offset. */
+	TERMINAL_API int32 GetAbsoluteRow(int32 ViewportRow, int32 ScrollOffset) const;
+
+	/** Clear the entire buffer. */
+	TERMINAL_API void Clear();
+
+	/** Get the number of rows available for scrollback (above the viewport). */
+	TERMINAL_API int32 GetScrollbackLength() const;
+
+	/** Get the total number of rows written (viewport + scrollback). */
+	TERMINAL_API int32 GetTotalRows() const;
+
+	int32 GetColumns() const { return Columns; }
+	int32 GetViewportRows() const { return ViewportRows; }
+
+	/** Scroll a region of viewport rows up by Count lines, filling new lines with blank cells. */
+	TERMINAL_API void ScrollRegionUp(int32 TopRow, int32 BottomRow, int32 Count);
+
+	/** Scroll a region of viewport rows down by Count lines, filling new lines with blank cells. */
+	TERMINAL_API void ScrollRegionDown(int32 TopRow, int32 BottomRow, int32 Count);
+
+	/** Fill a region of cells with a given cell value. */
+	TERMINAL_API void FillRegion(int32 AbsoluteRowStart, int32 ColumnStart, int32 AbsoluteRowEnd, int32 ColumnEnd, const FTerminalCell& Cell);
+
+	/** Get text content from a range of buffer-absolute rows. */
+	TERMINAL_API FString GetTextInRange(int32 StartRow, int32 StartColumn, int32 EndRow, int32 EndColumn) const;
+
+	/** Clear all scrollback rows, keeping only the current viewport content. */
+	TERMINAL_API void ClearScrollback();
+
+	FTerminalCursor Cursor;
+
+	/** Default cell used for clearing and new lines. */
+	FTerminalCell DefaultCell;
+
+	/** Alternate screen buffer support. */
+	TERMINAL_API void ActivateAlternateBuffer();
+	TERMINAL_API void DeactivateAlternateBuffer();
+	bool IsAlternateBufferActive() const { return bAlternateBufferActive; }
+
+private:
+
+	/** Get the ring buffer index for an absolute row. */
+	int32 GetRingIndex(int32 AbsoluteRow) const;
+
+	/** Ensure a row exists at the given ring index with correct column count. */
+	void EnsureRow(int32 RingIndex);
+
+	int32 Columns = 80;
+	int32 ViewportRows = 24;
+	int32 ScrollbackLimit = 131072;
+
+	/** Ring buffer of rows. Each row is a TArray<FTerminalCell> of Columns length. */
+	TArray<TArray<FTerminalCell>> Rows;
+	int32 RingStart = 0;
+	int32 TotalRowsWritten = 0;
+
+	/** Alternate buffer state. */
+	bool bAlternateBufferActive = false;
+	TArray<TArray<FTerminalCell>> SavedMainRows;
+	int32 SavedRingStart = 0;
+	int32 SavedTotalRowsWritten = 0;
+	FTerminalCursor SavedCursor;
+};
diff --git a/Engine/Plugins/Experimental/Terminal/Source/Terminal/Public/TerminalColorScheme.h b/Engine/Plugins/Experimental/Terminal/Source/Terminal/Public/TerminalColorScheme.h
new file mode 100644
--- /dev/null
+++ b/Engine/Plugins/Experimental/Terminal/Source/Terminal/Public/TerminalColorScheme.h
@@ -0,0 +1,27 @@
+// Copyright Epic Games, Inc. All Rights Reserved.
+
+#pragma once
+
+#include "CoreMinimal.h"
+
+/**
+ * Terminal color scheme loaded from a JSON file.
+ *
+ * Each scheme defines: default foreground/background, cursor color,
+ * selection highlight color, and a 16-color ANSI palette.
+ */
+struct FTerminalColorScheme
+{
+	FString Name;
+	FLinearColor DefaultForeground = FLinearColor(0.831f, 0.831f, 0.831f);
+	FLinearColor DefaultBackground = FLinearColor(0.118f, 0.118f, 0.118f);
+	FLinearColor CursorColor = FLinearColor(0.682f, 0.686f, 0.678f);
+	FLinearColor SelectionColor = FLinearColor(0.149f, 0.310f, 0.471f);
+	TArray<FLinearColor> Palette; // 16 ANSI colors (0-15)
+
+	/** Create a default scheme. */
+	static FTerminalColorScheme MakeDefault();
+
+	/** Parse a scheme from a JSON string. Returns true on success. */
+	static bool FromJSON(const FString& JSONString, FTerminalColorScheme& OutScheme);
+};
diff --git a/Engine/Plugins/Experimental/Terminal/Source/Terminal/Public/TerminalModule.h b/Engine/Plugins/Experimental/Terminal/Source/Terminal/Public/TerminalModule.h
new file mode 100644
--- /dev/null
+++ b/Engine/Plugins/Experimental/Terminal/Source/Terminal/Public/TerminalModule.h
@@ -0,0 +1,19 @@
+// Copyright Epic Games, Inc. All Rights Reserved.
+
+#pragma once
+
+#include "Modules/ModuleInterface.h"
+#include "Widgets/Docking/SDockTab.h"
+
+class FTerminalModule : public IModuleInterface
+{
+public:
+
+	virtual void StartupModule() override;
+
+	virtual void ShutdownModule() override;
+
+private:
+
+	TSharedRef<SDockTab> SpawnTerminalTab(const FSpawnTabArgs& SpawnTabArgs);
+};
diff --git a/Engine/Plugins/Experimental/Terminal/Source/Terminal/Public/TerminalSettings.h b/Engine/Plugins/Experimental/Terminal/Source/Terminal/Public/TerminalSettings.h
new file mode 100644
--- /dev/null
+++ b/Engine/Plugins/Experimental/Terminal/Source/Terminal/Public/TerminalSettings.h
@@ -0,0 +1,35 @@
+// Copyright Epic Games, Inc. All Rights Reserved.
+
+#pragma once
+
+#include "Engine/DeveloperSettings.h"
+
+#include "TerminalSettings.generated.h"
+
+UCLASS(MinimalAPI, config=EditorPerProjectUserSettings, defaultconfig, meta=(DisplayName="Terminal"))
+class UTerminalSettings : public UDeveloperSettings
+{
+	GENERATED_BODY()
+
+public:
+
+	/** Path to the shell executable. Leave empty to use the system default (COMSPEC on Windows, SHELL on Unix). */
+	UPROPERTY(Config, EditAnywhere, Category="Terminal")
+	FString ShellExecutablePath;
+
+	/** Font filename stem (without extension) as found in the system Fonts directory, e.g. "consola" for Consolas. */
+	UPROPERTY(Config, EditAnywhere, Category="Terminal")
+	FString FontFamily = TEXT("CascadiaMono");
+
+	/** Font size in points. */
+	UPROPERTY(Config, EditAnywhere, Category="Terminal", meta=(ClampMin=6, ClampMax=72))
+	int32 FontSize = 10;
+
+	/** Maximum number of scrollback rows retained by the terminal. */
+	UPROPERTY(Config, EditAnywhere, Category="Terminal", meta=(ClampMin=0, ClampMax=1000000))
+	int32 ScrollbackLimit = 131072;
+
+	/** Name of the color scheme to use. Must match a JSON file in Config/ColorSchemes/. */
+	UPROPERTY(Config, EditAnywhere, Category="Terminal")
+	FString ColorSchemeName = TEXT("Default");
+};
diff --git a/Engine/Plugins/Experimental/Terminal/Source/Terminal/Public/TerminalSubsystem.h b/Engine/Plugins/Experimental/Terminal/Source/Terminal/Public/TerminalSubsystem.h
new file mode 100644
--- /dev/null
+++ b/Engine/Plugins/Experimental/Terminal/Source/Terminal/Public/TerminalSubsystem.h
@@ -0,0 +1,41 @@
+// Copyright Epic Games, Inc. All Rights Reserved.
+
+#pragma once
+
+#include "EditorSubsystem.h"
+#include "TerminalColorScheme.h"
+
+#include "TerminalSubsystem.generated.h"
+
+/**
+ * Editor subsystem that manages terminal sessions and color schemes.
+ */
+UCLASS(MinimalAPI)
+class UTerminalSubsystem : public UEditorSubsystem
+{
+	GENERATED_BODY()
+
+public:
+
+	//~ Begin USubsystem interface
+	virtual void Initialize(FSubsystemCollectionBase& Collection) override;
+	virtual void Deinitialize() override;
+	//~ End USubsystem interface
+
+	/** Get the active color scheme based on settings. */
+	FTerminalColorScheme GetActiveColorScheme() const;
+
+	/** Get a color scheme by name. Returns the default if not found. */
+	FTerminalColorScheme GetColorScheme(const FString& Name) const;
+
+	/** Reload color schemes from disk. */
+	void ReloadColorSchemes();
+
+private:
+
+	/** Loaded color schemes indexed by name. */
+	TMap<FString, FTerminalColorScheme> ColorSchemes;
+
+	/** The hardcoded default scheme used as fallback. */
+	FTerminalColorScheme DefaultScheme;
+};
diff --git a/Engine/Plugins/Experimental/Terminal/Source/Terminal/Public/TerminalUtilities.h b/Engine/Plugins/Experimental/Terminal/Source/Terminal/Public/TerminalUtilities.h
new file mode 100644
--- /dev/null
+++ b/Engine/Plugins/Experimental/Terminal/Source/Terminal/Public/TerminalUtilities.h
@@ -0,0 +1,31 @@
+// Copyright Epic Games, Inc. All Rights Reserved.
+
+#pragma once
+
+#include "CoreMinimal.h"
+
+namespace TerminalUtilities
+{
+
+/** Returns true for characters that the primary monospace font will render with correct cell-width advance. */
+bool IsMonospaceSafe(TCHAR Character);
+
+/** Resolve a font family name to an absolute .ttf path in the system Fonts directory. Returns empty if not found. */
+FString ResolveSystemFontPath(const FString& FontFamily);
+
+/** Platform-appropriate monospace font fallback names, tried in order. */
+TArrayView<const TCHAR* const> GetFallbackFontNames();
+
+/** Platform-appropriate symbol font names (box drawing, arrows, math, etc.), tried in order. */
+TArrayView<const TCHAR* const> GetSymbolFontNames();
+
+/** Platform-appropriate emoji font names, tried in order. */
+TArrayView<const TCHAR* const> GetEmojiFontNames();
+
+/** Engine-bundled monospace font path, used as final fallback when no system font is found. */
+FString GetEngineFallbackFontPath();
+
+/** Parse a hex color string (with or without leading '#') into a linear color. Returns white on invalid input. */
+FLinearColor ParseHexColor(const FString& HexString);
+
+} // namespace TerminalUtilities
diff --git a/Engine/Plugins/Experimental/Terminal/Source/Terminal/Public/VTParser.h b/Engine/Plugins/Experimental/Terminal/Source/Terminal/Public/VTParser.h
new file mode 100644
--- /dev/null
+++ b/Engine/Plugins/Experimental/Terminal/Source/Terminal/Public/VTParser.h
@@ -0,0 +1,174 @@
+// Copyright Epic Games, Inc. All Rights Reserved.
+
+#pragma once
+
+#include "CoreMinimal.h"
+
+class FTerminalBuffer;
+
+/**
+ * VT/ANSI escape sequence parser implemented as a state machine.
+ *
+ * Follows the Paul Williams / vt100.net canonical parser model with 14 states,
+ * "anywhere" transitions (ESC/CAN/SUB), and proper C0 mid-sequence execution.
+ *
+ * Processes a byte buffer of UTF-8 text containing VT sequences and
+ * applies the resulting actions to an FTerminalBuffer.
+ */
+class FVTParser
+{
+public:
+
+	TERMINAL_API FVTParser();
+
+	/** Bind the parser to a terminal buffer. Must be called before Parse(). */
+	TERMINAL_API void SetBuffer(FTerminalBuffer* InBuffer);
+
+	/** Parse a chunk of UTF-8 bytes. */
+	TERMINAL_API void Parse(const uint8* Data, int32 Length);
+
+	/** Reset the parser and SGR state to defaults. */
+	TERMINAL_API void Reset();
+
+	/** Get the current parser state as a debug string. */
+	TERMINAL_API FString GetDebugStateName() const;
+
+	/** Initialize default tab stops (every 8 columns). */
+	TERMINAL_API void InitializeTabStops();
+
+	/** Current SGR state applied to newly written cells. */
+	FColor CurrentForeground;
+	FColor CurrentBackground;
+	uint8 CurrentAttributes = 0;
+
+	/** 16-color ANSI palette (indices 0-15). Initialized with defaults, overridden by color scheme. */
+	FColor AnsiPalette[16];
+
+	/** Scroll region (viewport-relative, inclusive). -1 means default (full viewport). */
+	int32 ScrollRegionTop = -1;
+	int32 ScrollRegionBottom = -1;
+
+	/** Mouse tracking mode set via DECSET 1000/1002/1003. Encoding format controlled by mode 1006 (`bSGRMouseEncoding`). */
+	enum class EMouseTrackingMode : uint8 { None, Normal, ButtonEvent, Any };
+
+	/** Mode flags. */
+	bool bAutoWrap = true;
+	bool bCursorVisible = true;
+	bool bApplicationCursorKeys = false;
+	bool bBracketedPaste = false;
+	bool bSynchronizedOutput = false;
+	bool bSGRMouseEncoding = false;
+	EMouseTrackingMode MouseTrackingMode = EMouseTrackingMode::None;
+
+	/** Tab stops. */
+	TArray<bool> TabStops;
+
+	/** Window title set via OSC. */
+	FString WindowTitle;
+
+	/** Response buffer  - bytes to send back to the PTY (DA, DSR replies). Drained by STerminal::OnTick. */
+	TArray<uint8> ResponseBuffer;
+
+private:
+
+	/** Parser states per the Williams/vt100.net state machine. */
+	enum class EParserState : uint8
+	{
+		Ground,
+		Escape,
+		EscapeIntermediate,
+		CSIEntry,
+		CSIParam,
+		CSIIntermediate,
+		CSIIgnore,
+		OSCString,
+		DCSEntry,
+		DCSParam,
+		DCSIntermediate,
+		DCSPassthrough,
+		DCSIgnore,
+		SOSPMAPCString,
+	};
+
+	EParserState State = EParserState::Ground;
+
+	/** CSI/DCS parameter accumulation. */
+	TArray<int32> CSIParams;
+	FString CSIIntermediateChars;
+	FString CurrentParam;
+
+	/** CSI private parameter prefix character ('<', '=', '>', '?'), or '\0' for none. */
+	TCHAR CSIPrivatePrefix = 0;
+
+	/** Escape sequence intermediate characters accumulator. */
+	FString EscapeIntermediateChars;
+
+	/** OSC accumulation. */
+	FString OSCPayload;
+
+	FTerminalBuffer* Buffer = nullptr;
+
+	/** UTF-8 decoder state. */
+	uint32 UTF8Codepoint = 0;
+	int32 UTF8Remaining = 0;
+	uint32 UTF8MinCodepoint = 0; // Minimum valid codepoint for current sequence length (overlong detection).
+
+	/** Process a single decoded Unicode codepoint. */
+	void ProcessCodepoint(uint32 Codepoint);
+
+	/** Reset CSI/DCS parameter accumulation state (the spec's "clear" action). */
+	void ClearSequenceState();
+
+	/** Execute a C0 control character in Ground context. Returns true if handled. */
+	bool TryExecuteC0(uint32 Codepoint);
+
+	/** State handlers. */
+	void HandleGround(uint32 Codepoint);
+	void HandleEscape(uint32 Codepoint);
+	void HandleEscapeIntermediate(uint32 Codepoint);
+	void HandleCSIEntry(uint32 Codepoint);
+	void HandleCSIParam(uint32 Codepoint);
+	void HandleCSIIntermediate(uint32 Codepoint);
+	void HandleCSIIgnore(uint32 Codepoint);
+	void HandleOSCString(uint32 Codepoint);
+	void HandleDCSEntry(uint32 Codepoint);
+	void HandleDCSParam(uint32 Codepoint);
+	void HandleDCSIntermediate(uint32 Codepoint);
+	void HandleDCSPassthrough(uint32 Codepoint);
+	void HandleDCSIgnore(uint32 Codepoint);
+	void HandleSOSPMAPCString(uint32 Codepoint);
+
+	/** Execute a completed CSI sequence. */
+	void ExecuteCSI(uint32 FinalByte);
+
+	/** Execute an SGR (Select Graphic Rendition) sequence. */
+	void ExecuteSGR();
+
+	/** Execute a DECSET/DECRST private mode sequence. */
+	void ExecutePrivateMode(uint32 FinalByte);
+
+	/** Execute a completed escape sequence with intermediates. */
+	void ExecuteEscapeSequence(uint32 FinalByte);
+
+	/** Execute a completed OSC sequence. */
+	void ExecuteOSC();
+
+	/** Write a printable character at the cursor position. */
+	void PutCharacter(TCHAR Character);
+
+	/** Check if a codepoint is zero-width (combining marks, variation selectors, etc.). */
+	static bool IsZeroWidthCodepoint(uint32 Codepoint);
+
+	/** Advance the cursor to the next line, scrolling if necessary. */
+	void LineFeed();
+
+	/** Get the effective scroll region boundaries. */
+	int32 GetEffectiveScrollTop() const;
+	int32 GetEffectiveScrollBottom() const;
+
+	/** Helper to get a CSI parameter with a default value. */
+	int32 GetCSIParam(int32 Index, int32 Default) const;
+
+	/** Get a color from the 256-color palette. Indices 0-15 use AnsiPalette. */
+	FColor Get256Color(int32 Index) const;
+};
diff --git a/Engine/Plugins/Experimental/Terminal/Source/Terminal/Terminal.Build.cs b/Engine/Plugins/Experimental/Terminal/Source/Terminal/Terminal.Build.cs
new file mode 100644
--- /dev/null
+++ b/Engine/Plugins/Experimental/Terminal/Source/Terminal/Terminal.Build.cs
@@ -0,0 +1,36 @@
+// Copyright Epic Games, Inc. All Rights Reserved.
+
+namespace UnrealBuildTool.Rules
+{
+	public class Terminal : ModuleRules
+	{
+		public Terminal(ReadOnlyTargetRules Target) : base(Target)
+		{
+			PrivateDependencyModuleNames.AddRange(
+				new string[] {
+					"ApplicationCore",
+					"Core",
+					"CoreUObject",
+					"DeveloperSettings",
+					"EditorSubsystem",
+					"Engine",
+					"InputCore",
+					"Json",
+					"JsonUtilities",
+					"Projects",
+					"Slate",
+					"SlateCore",
+					"ToolMenus",
+					"UnrealEd",
+					"WorkspaceMenuStructure",
+				}
+			);
+
+			// forkpty() on Linux lives in libutil.
+			if (Target.Platform.IsInGroup(UnrealPlatformGroup.Linux))
+			{
+				PublicSystemLibraries.Add("util");
+			}
+		}
+	}
+}
diff --git a/Engine/Plugins/Experimental/Terminal/Source/TerminalTests/Private/TerminalTestsModule.cpp b/Engine/Plugins/Experimental/Terminal/Source/TerminalTests/Private/TerminalTestsModule.cpp
new file mode 100644
--- /dev/null
+++ b/Engine/Plugins/Experimental/Terminal/Source/TerminalTests/Private/TerminalTestsModule.cpp
@@ -0,0 +1,9 @@
+// Copyright Epic Games, Inc. All Rights Reserved.
+
+#include "Modules/ModuleManager.h"
+
+class FTerminalTestsModule : public IModuleInterface
+{
+};
+
+IMPLEMENT_MODULE(FTerminalTestsModule, TerminalTests);
diff --git a/Engine/Plugins/Experimental/Terminal/Source/TerminalTests/Private/Tests/STerminalTests.cpp b/Engine/Plugins/Experimental/Terminal/Source/TerminalTests/Private/Tests/STerminalTests.cpp
new file mode 100644
--- /dev/null
+++ b/Engine/Plugins/Experimental/Terminal/Source/TerminalTests/Private/Tests/STerminalTests.cpp
@@ -0,0 +1,869 @@
+// Copyright Epic Games, Inc. All Rights Reserved.
+
+#include "Misc/AutomationTest.h"
+
+#include "STerminal.h"
+#include "TerminalBuffer.h"
+#include "VTParser.h"
+
+#include "Framework/Application/SlateApplication.h"
+#include "HAL/PlatformApplicationMisc.h"
+#include "Widgets/Layout/SScrollBar.h"
+
+#if WITH_DEV_AUTOMATION_TESTS
+
+// ── Helpers ──
+
+/** Feed a string as UTF-8 bytes into the parser. */
+static void ParseTerminalString(FVTParser& Parser, const FString& Text)
+{
+	const FTCHARToUTF8 Utf8(*Text);
+	Parser.Parse(reinterpret_cast<const uint8*>(Utf8.Get()), Utf8.Length());
+}
+
+/** Synthesize a keyboard event. */
+static FKeyEvent MakeKeyEvent(FKey Key, bool bShift = false, bool bCtrl = false)
+{
+	FModifierKeysState Modifiers(bShift, bShift, bCtrl, bCtrl, false, false, false, false, false);
+	return FKeyEvent(Key, Modifiers, 0, false, 0, 0);
+}
+
+/** Shared empty button set for synthesized pointer events. Static to avoid dangling references  - FPointerEvent stores a const TSet<FKey>& internally. */
+static const TSet<FKey> EmptyPressedButtons;
+
+/** Synthesize a mouse button event at a local pixel position. */
+static FPointerEvent MakeMouseEvent(const FVector2D& ScreenPosition, FKey Button = EKeys::LeftMouseButton, bool bShift = false)
+{
+	FModifierKeysState Modifiers(bShift, bShift, false, false, false, false, false, false, false);
+	return FPointerEvent(
+		0,
+		ScreenPosition,
+		ScreenPosition,
+		EmptyPressedButtons,
+		Button,
+		0.0f,
+		Modifiers);
+}
+
+/** Synthesize a mouse wheel event. */
+static FPointerEvent MakeWheelEvent(float WheelDelta, bool bCtrl = false)
+{
+	FModifierKeysState Modifiers(false, false, bCtrl, bCtrl, false, false, false, false, false);
+	return FPointerEvent(
+		0,
+		FVector2D::ZeroVector,
+		FVector2D::ZeroVector,
+		EmptyPressedButtons,
+		EKeys::Invalid,
+		WheelDelta,
+		Modifiers);
+}
+
+// ── Spec ──
+
+BEGIN_DEFINE_SPEC(FSTerminalSpec, "Terminal.STerminal",
+	EAutomationTestFlags::ProductFilter | EAutomationTestFlags_ApplicationContextMask)
+
+	TSharedPtr<STerminal> Terminal;
+
+	/** Initialize the terminal's internal buffer and state without spawning a real PTY session. */
+	void InitializeMinimalTerminal(int32 Columns = 80, int32 Rows = 24, int32 ScrollbackLimit = 100)
+	{
+		Terminal->ViewportColumns = Columns;
+		Terminal->ViewportRows = Rows;
+		Terminal->Buffer.Initialize(Columns, Rows, ScrollbackLimit);
+		Terminal->Parser.SetBuffer(&Terminal->Buffer);
+		Terminal->bInitialized = true;
+		Terminal->CellWidth = 8.0f;
+		Terminal->CellHeight = 16.0f;
+	}
+
+	/** Write text into the terminal buffer via the parser. */
+	void WriteText(const FString& Text)
+	{
+		ParseTerminalString(Terminal->Parser, Text);
+	}
+
+	/** Generate enough lines to fill the viewport and create scrollback. */
+	void GenerateScrollback(int32 LineCount = 50)
+	{
+		for (int32 Index = 0; Index < LineCount; ++Index)
+		{
+			WriteText(FString::Printf(TEXT("Line %d\r\n"), Index));
+		}
+	}
+
+	/** Create a geometry matching the terminal's viewport dimensions. */
+	FGeometry MakeGeometry() const
+	{
+		const float Width = Terminal->ViewportColumns * Terminal->CellWidth;
+		const float Height = Terminal->ViewportRows * Terminal->CellHeight;
+		return FGeometry::MakeRoot(FVector2D(Width, Height), FSlateLayoutTransform(FVector2D::ZeroVector));
+	}
+
+END_DEFINE_SPEC(FSTerminalSpec)
+
+void FSTerminalSpec::Define()
+{
+	BeforeEach([this]()
+	{
+		Terminal = SNew(STerminal);
+		InitializeMinimalTerminal();
+	});
+
+	AfterEach([this]()
+	{
+		Terminal.Reset();
+	});
+
+	// ═══════════════════════════════════════════════════════════════
+	// Foundational: coordinate mapping and buffer math
+	// ═══════════════════════════════════════════════════════════════
+
+	Describe("Buffer Coordinate Mapping", [this]()
+	{
+		It("should map viewport rows correctly with scrollback", [this]()
+		{
+			GenerateScrollback();
+			TestTrue(TEXT("has scrollback"), Terminal->Buffer.GetScrollbackLength() > 0);
+
+			const int32 TopRow = Terminal->Buffer.GetAbsoluteRow(0, 0);
+			const int32 BottomRow = Terminal->Buffer.GetAbsoluteRow(Terminal->ViewportRows - 1, 0);
+			TestEqual(TEXT("viewport spans ViewportRows"), BottomRow - TopRow, Terminal->ViewportRows - 1);
+		});
+
+		It("should shift rows with scroll offset", [this]()
+		{
+			GenerateScrollback();
+			const int32 RowAtOffset0 = Terminal->Buffer.GetAbsoluteRow(0, 0);
+			const int32 RowAtOffset5 = Terminal->Buffer.GetAbsoluteRow(0, 5);
+			TestEqual(TEXT("offset shifts by 5"), RowAtOffset0 - RowAtOffset5, 5);
+		});
+
+		It("should compute correct scrollback length", [this]()
+		{
+			TestEqual(TEXT("no scrollback initially"), Terminal->Buffer.GetScrollbackLength(), 0);
+
+			GenerateScrollback();
+			const int32 ExpectedScrollback = Terminal->Buffer.GetTotalRows() - Terminal->ViewportRows;
+			TestEqual(TEXT("scrollback = total - viewport"), Terminal->Buffer.GetScrollbackLength(), ExpectedScrollback);
+		});
+	});
+
+	Describe("PixelToCell", [this]()
+	{
+		It("should map top-left pixel to cell (0, topRow)", [this]()
+		{
+			const FGeometry Geometry = MakeGeometry();
+			const FIntPoint Cell = Terminal->PixelToCell(Geometry, FVector2D(0.0f, 0.0f));
+			TestEqual(TEXT("column"), Cell.X, 0);
+			TestEqual(TEXT("row"), Cell.Y, Terminal->Buffer.GetAbsoluteRow(0, Terminal->ScrollOffset));
+		});
+
+		It("should map bottom-right pixel to last cell", [this]()
+		{
+			const FGeometry Geometry = MakeGeometry();
+			const float MaxX = (Terminal->ViewportColumns - 1) * Terminal->CellWidth + 1.0f;
+			const float MaxY = (Terminal->ViewportRows - 1) * Terminal->CellHeight + 1.0f;
+			const FIntPoint Cell = Terminal->PixelToCell(Geometry, FVector2D(MaxX, MaxY));
+			TestEqual(TEXT("column"), Cell.X, Terminal->ViewportColumns - 1);
+			TestEqual(TEXT("row"), Cell.Y, Terminal->Buffer.GetAbsoluteRow(Terminal->ViewportRows - 1, Terminal->ScrollOffset));
+		});
+
+		It("should clamp negative coordinates to (0, topRow)", [this]()
+		{
+			const FGeometry Geometry = MakeGeometry();
+			const FIntPoint Cell = Terminal->PixelToCell(Geometry, FVector2D(-100.0f, -100.0f));
+			TestEqual(TEXT("column"), Cell.X, 0);
+			TestEqual(TEXT("row"), Cell.Y, Terminal->Buffer.GetAbsoluteRow(0, Terminal->ScrollOffset));
+		});
+
+		It("should clamp overflow coordinates to last cell", [this]()
+		{
+			const FGeometry Geometry = MakeGeometry();
+			const FIntPoint Cell = Terminal->PixelToCell(Geometry, FVector2D(9999.0f, 9999.0f));
+			TestEqual(TEXT("column"), Cell.X, Terminal->ViewportColumns - 1);
+			TestEqual(TEXT("row"), Cell.Y, Terminal->Buffer.GetAbsoluteRow(Terminal->ViewportRows - 1, Terminal->ScrollOffset));
+		});
+
+		It("should map mid-cell pixel to correct column", [this]()
+		{
+			const FGeometry Geometry = MakeGeometry();
+			const float MidX = 3.0f * Terminal->CellWidth + Terminal->CellWidth * 0.5f;
+			const FIntPoint Cell = Terminal->PixelToCell(Geometry, FVector2D(MidX, 0.0f));
+			TestEqual(TEXT("column"), Cell.X, 3);
+		});
+
+		It("should account for scroll offset", [this]()
+		{
+			GenerateScrollback();
+			Terminal->ScrollOffset = 5;
+			const FGeometry Geometry = MakeGeometry();
+			const FIntPoint Cell = Terminal->PixelToCell(Geometry, FVector2D(0.0f, 0.0f));
+			TestEqual(TEXT("scrolled row"), Cell.Y, Terminal->Buffer.GetAbsoluteRow(0, 5));
+		});
+	});
+
+	// ═══════════════════════════════════════════════════════════════
+	// Selection state: querying, text extraction, clearing
+	// ═══════════════════════════════════════════════════════════════
+
+	Describe("IsCellSelected", [this]()
+	{
+		It("should return false when no selection exists", [this]()
+		{
+			TestFalse(TEXT("no selection"), Terminal->IsCellSelected(0, 0));
+		});
+
+		It("should detect single-row selection", [this]()
+		{
+			Terminal->SelectionAnchor = FIntPoint(2, 0);
+			Terminal->SelectionActive = FIntPoint(5, 0);
+			Terminal->bHasSelection = true;
+
+			TestFalse(TEXT("before start"), Terminal->IsCellSelected(0, 1));
+			TestTrue(TEXT("at start"), Terminal->IsCellSelected(0, 2));
+			TestTrue(TEXT("middle"), Terminal->IsCellSelected(0, 3));
+			TestTrue(TEXT("at end"), Terminal->IsCellSelected(0, 5));
+			TestFalse(TEXT("after end"), Terminal->IsCellSelected(0, 6));
+		});
+
+		It("should detect multi-row selection", [this]()
+		{
+			Terminal->SelectionAnchor = FIntPoint(5, 1);
+			Terminal->SelectionActive = FIntPoint(3, 3);
+			Terminal->bHasSelection = true;
+
+			TestFalse(TEXT("row 0"), Terminal->IsCellSelected(0, 5));
+			TestFalse(TEXT("row 1, col 4"), Terminal->IsCellSelected(1, 4));
+			TestTrue(TEXT("row 1, col 5"), Terminal->IsCellSelected(1, 5));
+			TestTrue(TEXT("row 1, col 79"), Terminal->IsCellSelected(1, 79));
+			TestTrue(TEXT("row 2, col 0"), Terminal->IsCellSelected(2, 0));
+			TestTrue(TEXT("row 2, col 79"), Terminal->IsCellSelected(2, 79));
+			TestTrue(TEXT("row 3, col 0"), Terminal->IsCellSelected(3, 0));
+			TestTrue(TEXT("row 3, col 3"), Terminal->IsCellSelected(3, 3));
+			TestFalse(TEXT("row 3, col 4"), Terminal->IsCellSelected(3, 4));
+			TestFalse(TEXT("row 4"), Terminal->IsCellSelected(4, 0));
+		});
+
+		It("should handle reverse selection (anchor after active)", [this]()
+		{
+			Terminal->SelectionAnchor = FIntPoint(5, 2);
+			Terminal->SelectionActive = FIntPoint(2, 0);
+			Terminal->bHasSelection = true;
+
+			TestTrue(TEXT("start of reverse"), Terminal->IsCellSelected(0, 2));
+			TestTrue(TEXT("middle row"), Terminal->IsCellSelected(1, 0));
+			TestTrue(TEXT("end of reverse"), Terminal->IsCellSelected(2, 5));
+			TestFalse(TEXT("outside"), Terminal->IsCellSelected(2, 6));
+		});
+
+		It("should handle degenerate single-cell selection", [this]()
+		{
+			Terminal->SelectionAnchor = FIntPoint(3, 2);
+			Terminal->SelectionActive = FIntPoint(3, 2);
+			Terminal->bHasSelection = true;
+
+			TestTrue(TEXT("exact cell"), Terminal->IsCellSelected(2, 3));
+			TestFalse(TEXT("adjacent cell"), Terminal->IsCellSelected(2, 4));
+		});
+
+		It("should work during active drag (bSelecting)", [this]()
+		{
+			Terminal->SelectionAnchor = FIntPoint(0, 0);
+			Terminal->SelectionActive = FIntPoint(10, 0);
+			Terminal->bSelecting = true;
+			Terminal->bHasSelection = false;
+
+			TestTrue(TEXT("cell within drag"), Terminal->IsCellSelected(0, 5));
+		});
+	});
+
+	Describe("GetSelectedText", [this]()
+	{
+		It("should return empty when no selection", [this]()
+		{
+			TestEqual(TEXT("empty"), Terminal->GetSelectedText(), FString());
+		});
+
+		It("should return selected text from a single row", [this]()
+		{
+			WriteText(TEXT("Hello World"));
+			const int32 AbsoluteRow = Terminal->Buffer.GetAbsoluteRow(0, 0);
+			Terminal->SelectionAnchor = FIntPoint(0, AbsoluteRow);
+			Terminal->SelectionActive = FIntPoint(4, AbsoluteRow);
+			Terminal->bHasSelection = true;
+
+			TestEqual(TEXT("selected text"), Terminal->GetSelectedText(), TEXT("Hello"));
+		});
+
+		It("should return text across multiple rows", [this]()
+		{
+			WriteText(TEXT("AAAA\r\nBBBB\r\nCCCC"));
+			const int32 Row0 = Terminal->Buffer.GetAbsoluteRow(0, 0);
+			const int32 Row1 = Terminal->Buffer.GetAbsoluteRow(1, 0);
+			Terminal->SelectionAnchor = FIntPoint(2, Row0);
+			Terminal->SelectionActive = FIntPoint(1, Row1);
+			Terminal->bHasSelection = true;
+
+			const FString Selected = Terminal->GetSelectedText();
+			TestTrue(TEXT("starts with AA"), Selected.StartsWith(TEXT("AA")));
+			TestTrue(TEXT("ends with BB"), Selected.EndsWith(TEXT("BB")));
+		});
+
+		It("should return text from scrollback", [this]()
+		{
+			GenerateScrollback();
+			const int32 ScrollbackRow = Terminal->Buffer.GetAbsoluteRow(0, 10);
+			Terminal->SelectionAnchor = FIntPoint(0, ScrollbackRow);
+			Terminal->SelectionActive = FIntPoint(3, ScrollbackRow);
+			Terminal->bHasSelection = true;
+
+			TestTrue(TEXT("non-empty from scrollback"), Terminal->GetSelectedText().Len() > 0);
+		});
+	});
+
+	Describe("ClearSelection", [this]()
+	{
+		It("should reset all selection state", [this]()
+		{
+			Terminal->SelectionAnchor = FIntPoint(1, 2);
+			Terminal->SelectionActive = FIntPoint(3, 4);
+			Terminal->bHasSelection = true;
+			Terminal->bSelecting = true;
+			Terminal->SelectionAutoScrollDirection = 3;
+
+			Terminal->ClearSelection();
+
+			TestFalse(TEXT("bHasSelection"), Terminal->bHasSelection);
+			TestFalse(TEXT("bSelecting"), Terminal->bSelecting);
+			TestEqual(TEXT("auto-scroll cleared"), Terminal->SelectionAutoScrollDirection, 0);
+			TestEqual(TEXT("anchor reset"), Terminal->SelectionAnchor, FIntPoint(-1, -1));
+			TestEqual(TEXT("active reset"), Terminal->SelectionActive, FIntPoint(-1, -1));
+		});
+	});
+
+	// ═══════════════════════════════════════════════════════════════
+	// Click selection: single, double (word), triple (line)
+	// ═══════════════════════════════════════════════════════════════
+
+	Describe("Single Click", [this]()
+	{
+		It("should clear existing selection on single click", [this]()
+		{
+			Terminal->SelectionAnchor = FIntPoint(0, 0);
+			Terminal->SelectionActive = FIntPoint(10, 0);
+			Terminal->bHasSelection = true;
+
+			const FGeometry Geometry = MakeGeometry();
+			Terminal->OnMouseButtonDown(Geometry, MakeMouseEvent(FVector2D(50.0f, 50.0f)));
+			Terminal->OnMouseButtonUp(Geometry, MakeMouseEvent(FVector2D(50.0f, 50.0f)));
+
+			TestFalse(TEXT("selection cleared"), Terminal->bHasSelection);
+		});
+
+		It("should degenerate to no selection when start equals end", [this]()
+		{
+			const FGeometry Geometry = MakeGeometry();
+			const FVector2D Position(10.0f, 10.0f);
+
+			Terminal->OnMouseButtonDown(Geometry, MakeMouseEvent(Position));
+			Terminal->OnMouseButtonUp(Geometry, MakeMouseEvent(Position));
+
+			TestFalse(TEXT("no selection"), Terminal->bHasSelection);
+		});
+	});
+
+	Describe("Double-Click Word Selection", [this]()
+	{
+		It("should select a word on double-click", [this]()
+		{
+			WriteText(TEXT("Hello World"));
+			const FGeometry Geometry = MakeGeometry();
+
+			// Position over 'W' in "World" (column 6).
+			const FVector2D ClickPosition(6.0f * Terminal->CellWidth + 1.0f, 1.0f);
+
+			// First click.
+			Terminal->OnMouseButtonDown(Geometry, MakeMouseEvent(ClickPosition));
+			Terminal->OnMouseButtonUp(Geometry, MakeMouseEvent(ClickPosition));
+
+			// Second click (double-click).
+			Terminal->OnMouseButtonDown(Geometry, MakeMouseEvent(ClickPosition));
+
+			TestTrue(TEXT("has selection"), Terminal->bHasSelection);
+			TestEqual(TEXT("anchor col"), Terminal->SelectionAnchor.X, 6);
+			TestEqual(TEXT("active col"), Terminal->SelectionActive.X, 10);
+		});
+	});
+
+	Describe("Triple-Click Line Selection", [this]()
+	{
+		It("should select entire row on triple-click", [this]()
+		{
+			WriteText(TEXT("Hello World"));
+			const FGeometry Geometry = MakeGeometry();
+			const FVector2D ClickPosition(10.0f, 1.0f);
+
+			Terminal->OnMouseButtonDown(Geometry, MakeMouseEvent(ClickPosition));
+			Terminal->OnMouseButtonUp(Geometry, MakeMouseEvent(ClickPosition));
+			Terminal->OnMouseButtonDown(Geometry, MakeMouseEvent(ClickPosition));
+			Terminal->OnMouseButtonUp(Geometry, MakeMouseEvent(ClickPosition));
+			Terminal->OnMouseButtonDown(Geometry, MakeMouseEvent(ClickPosition));
+
+			TestTrue(TEXT("has selection"), Terminal->bHasSelection);
+			TestEqual(TEXT("anchor col 0"), Terminal->SelectionAnchor.X, 0);
+			TestEqual(TEXT("active col last"), Terminal->SelectionActive.X, Terminal->ViewportColumns - 1);
+		});
+	});
+
+	// ═══════════════════════════════════════════════════════════════
+	// Drag selection
+	// ═══════════════════════════════════════════════════════════════
+
+	Describe("Drag Selection", [this]()
+	{
+		It("should create selection when dragging across cells", [this]()
+		{
+			WriteText(TEXT("Hello World"));
+			const FGeometry Geometry = MakeGeometry();
+
+			const FVector2D StartPosition(0.0f, 1.0f);
+			const FVector2D EndPosition(4.0f * Terminal->CellWidth + 1.0f, 1.0f);
+
+			Terminal->OnMouseButtonDown(Geometry, MakeMouseEvent(StartPosition));
+			TestTrue(TEXT("selecting"), Terminal->bSelecting);
+
+			Terminal->OnMouseMove(Geometry, MakeMouseEvent(EndPosition));
+			TestNotEqual(TEXT("active moved"), Terminal->SelectionActive.X, Terminal->SelectionAnchor.X);
+
+			Terminal->OnMouseButtonUp(Geometry, MakeMouseEvent(EndPosition));
+			TestTrue(TEXT("has selection"), Terminal->bHasSelection);
+			TestFalse(TEXT("no longer selecting"), Terminal->bSelecting);
+		});
+
+		It("should create multi-row selection when dragging vertically", [this]()
+		{
+			WriteText(TEXT("Line 0\r\nLine 1\r\nLine 2\r\n"));
+			const FGeometry Geometry = MakeGeometry();
+
+			const FVector2D StartPosition(2.0f * Terminal->CellWidth, 0.5f * Terminal->CellHeight);
+			const FVector2D EndPosition(5.0f * Terminal->CellWidth, 2.5f * Terminal->CellHeight);
+
+			Terminal->OnMouseButtonDown(Geometry, MakeMouseEvent(StartPosition));
+			Terminal->OnMouseMove(Geometry, MakeMouseEvent(EndPosition));
+			Terminal->OnMouseButtonUp(Geometry, MakeMouseEvent(EndPosition));
+
+			TestTrue(TEXT("has selection"), Terminal->bHasSelection);
+			TestTrue(TEXT("spans multiple rows"), Terminal->SelectionAnchor.Y != Terminal->SelectionActive.Y);
+		});
+	});
+
+	// ═══════════════════════════════════════════════════════════════
+	// Drag auto-scroll (dragging above/below viewport)
+	// ═══════════════════════════════════════════════════════════════
+
+	Describe("Drag Auto-Scroll", [this]()
+	{
+		BeforeEach([this]()
+		{
+			GenerateScrollback();
+		});
+
+		It("should set positive auto-scroll when dragging above viewport", [this]()
+		{
+			const FGeometry Geometry = MakeGeometry();
+			Terminal->OnMouseButtonDown(Geometry, MakeMouseEvent(FVector2D(10.0f, 10.0f)));
+			Terminal->OnMouseMove(Geometry, MakeMouseEvent(FVector2D(10.0f, -20.0f)));
+
+			TestTrue(TEXT("auto-scroll direction > 0"), Terminal->SelectionAutoScrollDirection > 0);
+			TestTrue(TEXT("scroll offset increased"), Terminal->ScrollOffset > 0);
+		});
+
+		It("should set negative auto-scroll when dragging below viewport", [this]()
+		{
+			Terminal->ScrollOffset = 10;
+			const FGeometry Geometry = MakeGeometry();
+			Terminal->OnMouseButtonDown(Geometry, MakeMouseEvent(FVector2D(10.0f, 10.0f)));
+
+			const float BelowY = Terminal->ViewportRows * Terminal->CellHeight + 20.0f;
+			Terminal->OnMouseMove(Geometry, MakeMouseEvent(FVector2D(10.0f, BelowY)));
+
+			TestTrue(TEXT("auto-scroll direction < 0"), Terminal->SelectionAutoScrollDirection < 0);
+			TestTrue(TEXT("scroll offset decreased"), Terminal->ScrollOffset < 10);
+		});
+
+		It("should set SelectionActive to top viewport row when dragging above", [this]()
+		{
+			const FGeometry Geometry = MakeGeometry();
+			Terminal->OnMouseButtonDown(Geometry, MakeMouseEvent(FVector2D(10.0f, 10.0f)));
+			Terminal->OnMouseMove(Geometry, MakeMouseEvent(FVector2D(10.0f, -20.0f)));
+
+			const int32 TopVisibleRow = Terminal->Buffer.GetAbsoluteRow(0, Terminal->ScrollOffset);
+			TestEqual(TEXT("active at top of viewport"), Terminal->SelectionActive.Y, TopVisibleRow);
+		});
+
+		It("should set SelectionActive to bottom viewport row when dragging below", [this]()
+		{
+			Terminal->ScrollOffset = 10;
+			const FGeometry Geometry = MakeGeometry();
+			Terminal->OnMouseButtonDown(Geometry, MakeMouseEvent(FVector2D(10.0f, 10.0f)));
+
+			const float BelowY = Terminal->ViewportRows * Terminal->CellHeight + 20.0f;
+			Terminal->OnMouseMove(Geometry, MakeMouseEvent(FVector2D(10.0f, BelowY)));
+
+			const int32 BottomVisibleRow = Terminal->Buffer.GetAbsoluteRow(Terminal->ViewportRows - 1, Terminal->ScrollOffset);
+			TestEqual(TEXT("active at bottom of viewport"), Terminal->SelectionActive.Y, BottomVisibleRow);
+		});
+
+		It("should preserve column during auto-scroll", [this]()
+		{
+			const FGeometry Geometry = MakeGeometry();
+			const float TargetColumnX = 10.0f * Terminal->CellWidth + 1.0f;
+			Terminal->OnMouseButtonDown(Geometry, MakeMouseEvent(FVector2D(TargetColumnX, 10.0f)));
+			Terminal->OnMouseMove(Geometry, MakeMouseEvent(FVector2D(TargetColumnX, -20.0f)));
+
+			TestEqual(TEXT("column preserved"), Terminal->SelectionActive.X, 10);
+		});
+
+		It("should scale scroll speed with distance outside viewport", [this]()
+		{
+			const FGeometry Geometry = MakeGeometry();
+
+			// Slightly outside: -5px = less than one cell height above.
+			Terminal->OnMouseButtonDown(Geometry, MakeMouseEvent(FVector2D(10.0f, 10.0f)));
+			Terminal->OnMouseMove(Geometry, MakeMouseEvent(FVector2D(10.0f, -5.0f)));
+			const int32 SmallDirection = Terminal->SelectionAutoScrollDirection;
+
+			Terminal->OnMouseButtonUp(Geometry, MakeMouseEvent(FVector2D(10.0f, -5.0f)));
+
+			// Far outside: -80px = multiple cell heights above.
+			// Use a different click position to avoid double-click detection.
+			Terminal->ScrollOffset = 0;
+			Terminal->OnMouseButtonDown(Geometry, MakeMouseEvent(FVector2D(200.0f, 100.0f)));
+			Terminal->OnMouseMove(Geometry, MakeMouseEvent(FVector2D(10.0f, -80.0f)));
+			const int32 LargeDirection = Terminal->SelectionAutoScrollDirection;
+
+			TestTrue(TEXT("far outside scrolls faster"), LargeDirection >= SmallDirection);
+		});
+
+		It("should clear auto-scroll when mouse returns inside viewport", [this]()
+		{
+			const FGeometry Geometry = MakeGeometry();
+			Terminal->OnMouseButtonDown(Geometry, MakeMouseEvent(FVector2D(10.0f, 10.0f)));
+			Terminal->OnMouseMove(Geometry, MakeMouseEvent(FVector2D(10.0f, -20.0f)));
+			TestNotEqual(TEXT("auto-scroll set"), Terminal->SelectionAutoScrollDirection, 0);
+
+			Terminal->OnMouseMove(Geometry, MakeMouseEvent(FVector2D(10.0f, 10.0f)));
+			TestEqual(TEXT("auto-scroll cleared"), Terminal->SelectionAutoScrollDirection, 0);
+		});
+
+		It("should clear auto-scroll on mouse button up", [this]()
+		{
+			const FGeometry Geometry = MakeGeometry();
+			Terminal->OnMouseButtonDown(Geometry, MakeMouseEvent(FVector2D(10.0f, 10.0f)));
+			Terminal->OnMouseMove(Geometry, MakeMouseEvent(FVector2D(10.0f, -20.0f)));
+			Terminal->OnMouseButtonUp(Geometry, MakeMouseEvent(FVector2D(10.0f, -20.0f)));
+
+			TestEqual(TEXT("auto-scroll cleared"), Terminal->SelectionAutoScrollDirection, 0);
+		});
+
+		It("should clamp scroll offset when no scrollback available", [this]()
+		{
+			Terminal.Reset();
+			Terminal = SNew(STerminal);
+			InitializeMinimalTerminal();
+
+			const FGeometry Geometry = MakeGeometry();
+			Terminal->OnMouseButtonDown(Geometry, MakeMouseEvent(FVector2D(10.0f, 10.0f)));
+			Terminal->OnMouseMove(Geometry, MakeMouseEvent(FVector2D(10.0f, -50.0f)));
+
+			TestEqual(TEXT("offset stays 0"), Terminal->ScrollOffset, 0);
+		});
+	});
+
+	// ═══════════════════════════════════════════════════════════════
+	// Keyboard selection (Shift+Arrow)
+	// ═══════════════════════════════════════════════════════════════
+
+	Describe("Shift+Arrow Keyboard Selection", [this]()
+	{
+		BeforeEach([this]()
+		{
+			WriteText(TEXT("Line 0\r\nLine 1\r\nLine 2\r\nLine 3\r\n"));
+			Terminal->Buffer.Cursor.Row = 2;
+			Terminal->Buffer.Cursor.Column = 3;
+		});
+
+		// Vertical movement.
+
+		It("should start selection at cursor on first Shift+Up", [this]()
+		{
+			const FGeometry Geometry = MakeGeometry();
+			Terminal->OnKeyDown(Geometry, MakeKeyEvent(EKeys::Up, true));
+
+			TestTrue(TEXT("has selection"), Terminal->bHasSelection);
+			const int32 CursorAbsoluteRow = Terminal->Buffer.GetAbsoluteRow(2, 0);
+			TestEqual(TEXT("anchor row"), Terminal->SelectionAnchor.Y, CursorAbsoluteRow);
+			TestEqual(TEXT("anchor col"), Terminal->SelectionAnchor.X, 3);
+			TestEqual(TEXT("active row moved up"), Terminal->SelectionActive.Y, CursorAbsoluteRow - 1);
+		});
+
+		It("should start selection at cursor on first Shift+Down", [this]()
+		{
+			const FGeometry Geometry = MakeGeometry();
+			Terminal->OnKeyDown(Geometry, MakeKeyEvent(EKeys::Down, true));
+
+			TestTrue(TEXT("has selection"), Terminal->bHasSelection);
+			const int32 CursorAbsoluteRow = Terminal->Buffer.GetAbsoluteRow(2, 0);
+			TestEqual(TEXT("active row moved down"), Terminal->SelectionActive.Y, CursorAbsoluteRow + 1);
+		});
+
+		It("should extend selection with multiple Shift+Up presses", [this]()
+		{
+			// Cursor at row 2  - press Shift+Up twice to reach row 0.
+			const FGeometry Geometry = MakeGeometry();
+			Terminal->OnKeyDown(Geometry, MakeKeyEvent(EKeys::Up, true));
+			Terminal->OnKeyDown(Geometry, MakeKeyEvent(EKeys::Up, true));
+
+			const int32 CursorAbsoluteRow = Terminal->Buffer.GetAbsoluteRow(2, 0);
+			TestEqual(TEXT("active moved 2 rows up"), Terminal->SelectionActive.Y, CursorAbsoluteRow - 2);
+			TestEqual(TEXT("anchor unchanged"), Terminal->SelectionAnchor.Y, CursorAbsoluteRow);
+		});
+
+		It("should not extend selection past last row with Shift+Down", [this]()
+		{
+			Terminal->Buffer.Cursor.Row = 3;
+			Terminal->Buffer.Cursor.Column = 0;
+			const int32 LastRow = Terminal->Buffer.GetTotalRows() - 1;
+
+			const FGeometry Geometry = MakeGeometry();
+			for (int32 Index = 0; Index < 30; ++Index)
+			{
+				Terminal->OnKeyDown(Geometry, MakeKeyEvent(EKeys::Down, true));
+			}
+
+			TestTrue(TEXT("has selection"), Terminal->bHasSelection);
+			TestTrue(TEXT("active at or before last row"), Terminal->SelectionActive.Y <= LastRow);
+		});
+
+		// Horizontal movement.
+
+		It("should move Shift+Left within a row", [this]()
+		{
+			const FGeometry Geometry = MakeGeometry();
+			Terminal->OnKeyDown(Geometry, MakeKeyEvent(EKeys::Left, true));
+
+			TestTrue(TEXT("has selection"), Terminal->bHasSelection);
+			TestEqual(TEXT("active col"), Terminal->SelectionActive.X, 2);
+			TestEqual(TEXT("same row"), Terminal->SelectionActive.Y, Terminal->SelectionAnchor.Y);
+		});
+
+		It("should move Shift+Right within a row", [this]()
+		{
+			const FGeometry Geometry = MakeGeometry();
+			Terminal->OnKeyDown(Geometry, MakeKeyEvent(EKeys::Right, true));
+
+			TestTrue(TEXT("has selection"), Terminal->bHasSelection);
+			TestEqual(TEXT("active col"), Terminal->SelectionActive.X, 4);
+			TestEqual(TEXT("same row"), Terminal->SelectionActive.Y, Terminal->SelectionAnchor.Y);
+		});
+
+		// Line wrapping.
+
+		It("should handle Shift+Left wrapping to previous row", [this]()
+		{
+			Terminal->Buffer.Cursor.Column = 0;
+			const FGeometry Geometry = MakeGeometry();
+			Terminal->OnKeyDown(Geometry, MakeKeyEvent(EKeys::Left, true));
+
+			TestTrue(TEXT("has selection"), Terminal->bHasSelection);
+			const int32 CursorAbsoluteRow = Terminal->Buffer.GetAbsoluteRow(2, 0);
+			TestEqual(TEXT("active row wrapped up"), Terminal->SelectionActive.Y, CursorAbsoluteRow - 1);
+			TestEqual(TEXT("active col at end"), Terminal->SelectionActive.X, Terminal->ViewportColumns - 1);
+		});
+
+		It("should handle Shift+Right wrapping to next row", [this]()
+		{
+			Terminal->Buffer.Cursor.Column = Terminal->ViewportColumns - 1;
+			const FGeometry Geometry = MakeGeometry();
+			Terminal->OnKeyDown(Geometry, MakeKeyEvent(EKeys::Right, true));
+
+			TestTrue(TEXT("has selection"), Terminal->bHasSelection);
+			const int32 CursorAbsoluteRow = Terminal->Buffer.GetAbsoluteRow(2, 0);
+			TestEqual(TEXT("active row wrapped down"), Terminal->SelectionActive.Y, CursorAbsoluteRow + 1);
+			TestEqual(TEXT("active col at start"), Terminal->SelectionActive.X, 0);
+		});
+
+		// Viewport scrolling.
+
+		It("should scroll viewport when selection extends above visible area", [this]()
+		{
+			Terminal.Reset();
+			Terminal = SNew(STerminal);
+			InitializeMinimalTerminal(80, 5, 100);
+			for (int32 Index = 0; Index < 20; ++Index)
+			{
+				WriteText(FString::Printf(TEXT("Line %d\r\n"), Index));
+			}
+			Terminal->ScrollOffset = 0;
+			Terminal->Buffer.Cursor.Row = 0;
+			Terminal->Buffer.Cursor.Column = 0;
+
+			const FGeometry Geometry = FGeometry::MakeRoot(
+				FVector2D(Terminal->ViewportColumns * Terminal->CellWidth, Terminal->ViewportRows * Terminal->CellHeight),
+				FSlateLayoutTransform(FVector2D::ZeroVector));
+
+			Terminal->OnKeyDown(Geometry, MakeKeyEvent(EKeys::Up, true));
+			Terminal->OnKeyDown(Geometry, MakeKeyEvent(EKeys::Up, true));
+
+			TestTrue(TEXT("viewport scrolled"), Terminal->ScrollOffset > 0);
+		});
+	});
+
+	// ═══════════════════════════════════════════════════════════════
+	// Mouse wheel scrolling
+	// ═══════════════════════════════════════════════════════════════
+
+	Describe("Mouse Wheel Scrolling", [this]()
+	{
+		It("should scroll up on positive wheel delta", [this]()
+		{
+			GenerateScrollback();
+			TestEqual(TEXT("initial scroll offset"), Terminal->ScrollOffset, 0);
+
+			const FGeometry Geometry = MakeGeometry();
+			Terminal->OnMouseWheel(Geometry, MakeWheelEvent(1.0f));
+
+			TestTrue(TEXT("scrolled up"), Terminal->ScrollOffset > 0);
+		});
+
+		It("should not scroll below zero", [this]()
+		{
+			Terminal->ScrollOffset = 0;
+			const FGeometry Geometry = MakeGeometry();
+			Terminal->OnMouseWheel(Geometry, MakeWheelEvent(-1.0f));
+
+			TestEqual(TEXT("clamped to zero"), Terminal->ScrollOffset, 0);
+		});
+
+		It("should not scroll above max scrollback", [this]()
+		{
+			GenerateScrollback();
+			const int32 MaxScrollback = Terminal->Buffer.GetScrollbackLength();
+			Terminal->ScrollOffset = MaxScrollback;
+
+			const FGeometry Geometry = MakeGeometry();
+			Terminal->OnMouseWheel(Geometry, MakeWheelEvent(100.0f));
+
+			TestEqual(TEXT("clamped to max"), Terminal->ScrollOffset, MaxScrollback);
+		});
+
+		It("should preserve selection when scrolling", [this]()
+		{
+			GenerateScrollback();
+			const int32 AbsoluteRow = Terminal->Buffer.GetAbsoluteRow(0, 0);
+			Terminal->SelectionAnchor = FIntPoint(0, AbsoluteRow);
+			Terminal->SelectionActive = FIntPoint(5, AbsoluteRow);
+			Terminal->bHasSelection = true;
+
+			const FGeometry Geometry = MakeGeometry();
+			Terminal->OnMouseWheel(Geometry, MakeWheelEvent(3.0f));
+
+			TestTrue(TEXT("selection preserved"), Terminal->bHasSelection);
+			TestEqual(TEXT("anchor unchanged"), Terminal->SelectionAnchor, FIntPoint(0, AbsoluteRow));
+			TestEqual(TEXT("active unchanged"), Terminal->SelectionActive, FIntPoint(5, AbsoluteRow));
+		});
+	});
+
+	// ═══════════════════════════════════════════════════════════════
+	// Widget lifecycle: focus
+	// ═══════════════════════════════════════════════════════════════
+
+	Describe("Focus", [this]()
+	{
+		It("should set focus flag on focus received", [this]()
+		{
+			const FGeometry Geometry = MakeGeometry();
+			Terminal->OnFocusReceived(Geometry, FFocusEvent());
+
+			TestTrue(TEXT("has focus"), Terminal->bHasFocus);
+		});
+
+		It("should clear drag state on focus lost", [this]()
+		{
+			Terminal->bSelecting = true;
+			Terminal->SelectionAutoScrollDirection = 3;
+
+			Terminal->OnFocusLost(FFocusEvent());
+
+			TestFalse(TEXT("selecting cleared"), Terminal->bSelecting);
+			TestEqual(TEXT("auto-scroll cleared"), Terminal->SelectionAutoScrollDirection, 0);
+			TestFalse(TEXT("focus flag"), Terminal->bHasFocus);
+		});
+	});
+
+	// ═══════════════════════════════════════════════════════════════
+	// Alternate buffer scrolling
+	// ═══════════════════════════════════════════════════════════════
+
+	Describe("Alternate Buffer Scrolling", [this]()
+	{
+		It("should not scroll in alternate screen buffer", [this]()
+		{
+			GenerateScrollback();
+			WriteText(TEXT("\x1B[?1049h")); // Activate alternate buffer
+			Terminal->ScrollOffset = 0;
+
+			const FGeometry Geometry = MakeGeometry();
+			Terminal->OnMouseWheel(Geometry, MakeWheelEvent(3.0f));
+
+			TestEqual(TEXT("scroll offset stays 0"), Terminal->ScrollOffset, 0);
+		});
+
+		It("should scroll normally after leaving alternate buffer", [this]()
+		{
+			WriteText(TEXT("\x1B[?1049h")); // Activate
+			WriteText(TEXT("\x1B[?1049l")); // Deactivate
+			GenerateScrollback();
+
+			const FGeometry Geometry = MakeGeometry();
+			Terminal->OnMouseWheel(Geometry, MakeWheelEvent(1.0f));
+
+			TestTrue(TEXT("scrolled up"), Terminal->ScrollOffset > 0);
+		});
+	});
+
+	// ═══════════════════════════════════════════════════════════════
+	// Mouse tracking
+	// ═══════════════════════════════════════════════════════════════
+
+	Describe("Mouse Tracking", [this]()
+	{
+		It("should not adjust ScrollOffset when mouse tracking is active", [this]()
+		{
+			GenerateScrollback();
+			WriteText(TEXT("\x1B[?1000h")); // Enable normal mouse tracking
+			Terminal->ScrollOffset = 0;
+
+			const FGeometry Geometry = MakeGeometry();
+			Terminal->OnMouseWheel(Geometry, MakeWheelEvent(3.0f));
+
+			TestEqual(TEXT("scroll offset stays 0"), Terminal->ScrollOffset, 0);
+		});
+
+		It("should resume normal scrolling when mouse tracking is disabled", [this]()
+		{
+			GenerateScrollback();
+			WriteText(TEXT("\x1B[?1000h")); // Enable
+			WriteText(TEXT("\x1B[?1000l")); // Disable
+
+			const FGeometry Geometry = MakeGeometry();
+			Terminal->OnMouseWheel(Geometry, MakeWheelEvent(1.0f));
+
+			TestTrue(TEXT("scrolled up"), Terminal->ScrollOffset > 0);
+		});
+	});
+}
+
+#endif // WITH_DEV_AUTOMATION_TESTS
diff --git a/Engine/Plugins/Experimental/Terminal/Source/TerminalTests/Private/Tests/VTParserTests.cpp b/Engine/Plugins/Experimental/Terminal/Source/TerminalTests/Private/Tests/VTParserTests.cpp
new file mode 100644
--- /dev/null
+++ b/Engine/Plugins/Experimental/Terminal/Source/TerminalTests/Private/Tests/VTParserTests.cpp
@@ -0,0 +1,899 @@
+// Copyright Epic Games, Inc. All Rights Reserved.
+
+#include "Misc/AutomationTest.h"
+
+#include "TerminalBuffer.h"
+#include "VTParser.h"
+
+#if WITH_DEV_AUTOMATION_TESTS
+
+// ── Helpers ──
+
+/** Feed a raw byte sequence into the parser. */
+static void ParseBytes(FVTParser& Parser, TArrayView<const uint8> Bytes)
+{
+	Parser.Parse(Bytes.GetData(), Bytes.Num());
+}
+
+/** Feed a string as UTF-8 bytes into the parser. */
+static void ParseString(FVTParser& Parser, const FString& Text)
+{
+	const FTCHARToUTF8 Utf8(*Text);
+	Parser.Parse(reinterpret_cast<const uint8*>(Utf8.Get()), Utf8.Length());
+}
+
+/** Get the character at a viewport-relative position. */
+static TCHAR GetCharAt(FTerminalBuffer& Buffer, int32 ViewportRow, int32 Column)
+{
+	const int32 AbsoluteRow = Buffer.GetAbsoluteRow(ViewportRow, 0);
+	return Buffer.GetCell(AbsoluteRow, Column).Character;
+}
+
+/** Get the cell at a viewport-relative position. */
+static const FTerminalCell& GetCellAt(FTerminalBuffer& Buffer, int32 ViewportRow, int32 Column)
+{
+	const int32 AbsoluteRow = Buffer.GetAbsoluteRow(ViewportRow, 0);
+	return Buffer.GetCell(AbsoluteRow, Column);
+}
+
+/** Get the text content of a viewport row. */
+static FString GetRowText(FTerminalBuffer& Buffer, int32 ViewportRow)
+{
+	const int32 AbsoluteRow = Buffer.GetAbsoluteRow(ViewportRow, 0);
+	return Buffer.GetTextInRange(AbsoluteRow, 0, AbsoluteRow, Buffer.GetColumns() - 1);
+}
+
+// ── Spec ──
+
+BEGIN_DEFINE_SPEC(FVTParserSpec, "Terminal.VTParser",
+	EAutomationTestFlags::ProductFilter | EAutomationTestFlags_ApplicationContextMask)
+
+	TUniquePtr<FTerminalBuffer> Buffer;
+	TUniquePtr<FVTParser> Parser;
+
+END_DEFINE_SPEC(FVTParserSpec)
+
+void FVTParserSpec::Define()
+{
+	BeforeEach([this]()
+	{
+		Buffer = MakeUnique<FTerminalBuffer>();
+		Buffer->Initialize(80, 24, 0);
+		Parser = MakeUnique<FVTParser>();
+		Parser->SetBuffer(Buffer.Get());
+	});
+
+	AfterEach([this]()
+	{
+		Parser.Reset();
+		Buffer.Reset();
+	});
+
+	// ── A. Plain Text & C0 Controls ──
+
+	Describe("Plain Text & C0 Controls", [this]()
+	{
+		It("should print ASCII text at cursor and advance cursor", [this]()
+		{
+			ParseString(*Parser, TEXT("Hello"));
+			TestEqual(TEXT("char at (0,0)"), GetCharAt(*Buffer, 0, 0), TEXT('H'));
+			TestEqual(TEXT("char at (0,4)"), GetCharAt(*Buffer, 0, 4), TEXT('o'));
+			TestEqual(TEXT("cursor column"), Buffer->Cursor.Column, 5);
+			TestEqual(TEXT("cursor row"), Buffer->Cursor.Row, 0);
+		});
+
+		It("should handle CR by moving cursor to column 0", [this]()
+		{
+			ParseString(*Parser, TEXT("Hello"));
+			ParseBytes(*Parser, { 0x0D }); // CR
+			TestEqual(TEXT("cursor column after CR"), Buffer->Cursor.Column, 0);
+			TestEqual(TEXT("cursor row after CR"), Buffer->Cursor.Row, 0);
+		});
+
+		It("should handle LF by moving cursor down one row", [this]()
+		{
+			ParseString(*Parser, TEXT("A"));
+			ParseBytes(*Parser, { 0x0A }); // LF
+			TestEqual(TEXT("cursor row after LF"), Buffer->Cursor.Row, 1);
+		});
+
+		It("should handle VT and FF as LF", [this]()
+		{
+			ParseBytes(*Parser, { 0x0B }); // VT
+			TestEqual(TEXT("cursor row after VT"), Buffer->Cursor.Row, 1);
+			ParseBytes(*Parser, { 0x0C }); // FF
+			TestEqual(TEXT("cursor row after FF"), Buffer->Cursor.Row, 2);
+		});
+
+		It("should scroll when LF at bottom of viewport", [this]()
+		{
+			// Move cursor to last row.
+			for (int32 Index = 0; Index < 23; ++Index)
+			{
+				ParseBytes(*Parser, { 0x0A });
+			}
+			TestEqual(TEXT("cursor at bottom row"), Buffer->Cursor.Row, 23);
+
+			ParseString(*Parser, TEXT("bottom"));
+			ParseBytes(*Parser, { 0x0A }); // LF at bottom should scroll
+
+			// Cursor should still be at row 23.
+			TestEqual(TEXT("cursor row after scroll"), Buffer->Cursor.Row, 23);
+		});
+
+		It("should handle BS by moving cursor left, stopping at column 0", [this]()
+		{
+			ParseString(*Parser, TEXT("AB"));
+			ParseBytes(*Parser, { 0x08 }); // BS
+			TestEqual(TEXT("cursor after BS"), Buffer->Cursor.Column, 1);
+
+			// BS at column 0 should stop.
+			Buffer->Cursor.Column = 0;
+			ParseBytes(*Parser, { 0x08 });
+			TestEqual(TEXT("cursor stays at 0"), Buffer->Cursor.Column, 0);
+		});
+
+		It("should handle HT by advancing to next tab stop", [this]()
+		{
+			ParseBytes(*Parser, { 0x09 }); // HT from column 0
+			TestEqual(TEXT("tab to column 8"), Buffer->Cursor.Column, 8);
+
+			ParseBytes(*Parser, { 0x09 }); // HT from column 8
+			TestEqual(TEXT("tab to column 16"), Buffer->Cursor.Column, 16);
+		});
+
+		It("should handle BEL without crashing", [this]()
+		{
+			ParseBytes(*Parser, { 0x07 });
+			TestEqual(TEXT("cursor unchanged"), Buffer->Cursor.Column, 0);
+		});
+	});
+
+	// ── B. CSI Cursor Movement ──
+
+	Describe("CSI Cursor Movement", [this]()
+	{
+		It("CUU: should move cursor up", [this]()
+		{
+			Buffer->Cursor.Row = 5;
+			ParseString(*Parser, TEXT("\x1B[3A")); // CUU 3
+			TestEqual(TEXT("cursor row"), Buffer->Cursor.Row, 2);
+		});
+
+		It("CUU: should stop at top row", [this]()
+		{
+			Buffer->Cursor.Row = 2;
+			ParseString(*Parser, TEXT("\x1B[10A"));
+			TestEqual(TEXT("cursor row clamped at top"), Buffer->Cursor.Row, 0);
+		});
+
+		It("CUU: default parameter should be 1", [this]()
+		{
+			Buffer->Cursor.Row = 3;
+			ParseString(*Parser, TEXT("\x1B[A"));
+			TestEqual(TEXT("cursor row"), Buffer->Cursor.Row, 2);
+		});
+
+		It("CUD: should move cursor down", [this]()
+		{
+			ParseString(*Parser, TEXT("\x1B[5B")); // CUD 5
+			TestEqual(TEXT("cursor row"), Buffer->Cursor.Row, 5);
+		});
+
+		It("CUD: should stop at bottom row", [this]()
+		{
+			ParseString(*Parser, TEXT("\x1B[100B"));
+			TestEqual(TEXT("cursor row clamped at bottom"), Buffer->Cursor.Row, 23);
+		});
+
+		It("CUF: should move cursor right", [this]()
+		{
+			ParseString(*Parser, TEXT("\x1B[10C")); // CUF 10
+			TestEqual(TEXT("cursor column"), Buffer->Cursor.Column, 10);
+		});
+
+		It("CUF: should stop at right edge", [this]()
+		{
+			ParseString(*Parser, TEXT("\x1B[200C"));
+			TestEqual(TEXT("cursor column clamped at right"), Buffer->Cursor.Column, 79);
+		});
+
+		It("CUB: should move cursor left", [this]()
+		{
+			Buffer->Cursor.Column = 10;
+			ParseString(*Parser, TEXT("\x1B[3D")); // CUB 3
+			TestEqual(TEXT("cursor column"), Buffer->Cursor.Column, 7);
+		});
+
+		It("CUB: should stop at left edge", [this]()
+		{
+			Buffer->Cursor.Column = 2;
+			ParseString(*Parser, TEXT("\x1B[10D"));
+			TestEqual(TEXT("cursor column clamped at left"), Buffer->Cursor.Column, 0);
+		});
+
+		It("CHA: should set cursor to absolute column (1-based)", [this]()
+		{
+			ParseString(*Parser, TEXT("\x1B[20G")); // CHA column 20
+			TestEqual(TEXT("cursor column"), Buffer->Cursor.Column, 19);
+		});
+
+		It("CHA: default should be column 1", [this]()
+		{
+			Buffer->Cursor.Column = 10;
+			ParseString(*Parser, TEXT("\x1B[G"));
+			TestEqual(TEXT("cursor column"), Buffer->Cursor.Column, 0);
+		});
+
+		It("CUP: should set cursor to row;col (1-based)", [this]()
+		{
+			ParseString(*Parser, TEXT("\x1B[10;20H")); // CUP row 10, col 20
+			TestEqual(TEXT("cursor row"), Buffer->Cursor.Row, 9);
+			TestEqual(TEXT("cursor column"), Buffer->Cursor.Column, 19);
+		});
+
+		It("CUP: default should be 1;1", [this]()
+		{
+			Buffer->Cursor.Row = 5;
+			Buffer->Cursor.Column = 10;
+			ParseString(*Parser, TEXT("\x1B[H"));
+			TestEqual(TEXT("cursor row"), Buffer->Cursor.Row, 0);
+			TestEqual(TEXT("cursor column"), Buffer->Cursor.Column, 0);
+		});
+	});
+
+	// ── C. CSI Erase ──
+
+	Describe("CSI Erase", [this]()
+	{
+		It("ED 0: should erase from cursor to end of display", [this]()
+		{
+			// Fill row 0 with text.
+			ParseString(*Parser, TEXT("ABCDEFGH"));
+			// Move cursor to (0, 4).
+			ParseString(*Parser, TEXT("\x1B[1;5H"));
+			// Erase from cursor to end of display.
+			ParseString(*Parser, TEXT("\x1B[0J"));
+
+			TestEqual(TEXT("preserved before cursor"), GetCharAt(*Buffer, 0, 0), TEXT('A'));
+			TestEqual(TEXT("preserved before cursor"), GetCharAt(*Buffer, 0, 3), TEXT('D'));
+			TestEqual(TEXT("erased at cursor"), GetCharAt(*Buffer, 0, 4), TEXT(' '));
+			TestEqual(TEXT("erased after cursor"), GetCharAt(*Buffer, 0, 7), TEXT(' '));
+		});
+
+		It("ED 1: should erase from start of display to cursor", [this]()
+		{
+			ParseString(*Parser, TEXT("ABCDEFGH"));
+			ParseString(*Parser, TEXT("\x1B[1;5H")); // cursor at (0, 4)
+			ParseString(*Parser, TEXT("\x1B[1J"));
+
+			TestEqual(TEXT("erased before cursor"), GetCharAt(*Buffer, 0, 0), TEXT(' '));
+			TestEqual(TEXT("erased at cursor"), GetCharAt(*Buffer, 0, 4), TEXT(' '));
+			TestEqual(TEXT("preserved after cursor"), GetCharAt(*Buffer, 0, 5), TEXT('F'));
+		});
+
+		It("ED 2: should erase entire display", [this]()
+		{
+			ParseString(*Parser, TEXT("Hello"));
+			ParseString(*Parser, TEXT("\x1B[2J"));
+			TestEqual(TEXT("cell erased"), GetCharAt(*Buffer, 0, 0), TEXT(' '));
+		});
+
+		It("EL 0: should erase from cursor to end of line", [this]()
+		{
+			ParseString(*Parser, TEXT("ABCDEFGH"));
+			ParseString(*Parser, TEXT("\x1B[1;5H")); // cursor at (0, 4)
+			ParseString(*Parser, TEXT("\x1B[0K"));
+
+			TestEqual(TEXT("preserved"), GetCharAt(*Buffer, 0, 3), TEXT('D'));
+			TestEqual(TEXT("erased at cursor"), GetCharAt(*Buffer, 0, 4), TEXT(' '));
+			TestEqual(TEXT("erased after"), GetCharAt(*Buffer, 0, 7), TEXT(' '));
+		});
+
+		It("EL 1: should erase from start of line to cursor", [this]()
+		{
+			ParseString(*Parser, TEXT("ABCDEFGH"));
+			ParseString(*Parser, TEXT("\x1B[1;5H")); // cursor at (0, 4)
+			ParseString(*Parser, TEXT("\x1B[1K"));
+
+			TestEqual(TEXT("erased before"), GetCharAt(*Buffer, 0, 0), TEXT(' '));
+			TestEqual(TEXT("erased at cursor"), GetCharAt(*Buffer, 0, 4), TEXT(' '));
+			TestEqual(TEXT("preserved after"), GetCharAt(*Buffer, 0, 5), TEXT('F'));
+		});
+
+		It("EL 2: should erase entire line", [this]()
+		{
+			ParseString(*Parser, TEXT("ABCDEFGH"));
+			ParseString(*Parser, TEXT("\x1B[1;5H"));
+			ParseString(*Parser, TEXT("\x1B[2K"));
+
+			TestEqual(TEXT("erased start"), GetCharAt(*Buffer, 0, 0), TEXT(' '));
+			TestEqual(TEXT("erased end"), GetCharAt(*Buffer, 0, 7), TEXT(' '));
+		});
+	});
+
+	// ── D. CSI Insert/Delete ──
+
+	Describe("CSI Insert/Delete", [this]()
+	{
+		It("ICH: should insert blank characters at cursor", [this]()
+		{
+			ParseString(*Parser, TEXT("ABCDEF"));
+			ParseString(*Parser, TEXT("\x1B[1;3H")); // cursor at (0, 2)
+			ParseString(*Parser, TEXT("\x1B[2@"));    // ICH 2
+
+			TestEqual(TEXT("A preserved"), GetCharAt(*Buffer, 0, 0), TEXT('A'));
+			TestEqual(TEXT("B preserved"), GetCharAt(*Buffer, 0, 1), TEXT('B'));
+			TestEqual(TEXT("blank inserted"), GetCharAt(*Buffer, 0, 2), TEXT(' '));
+			TestEqual(TEXT("blank inserted"), GetCharAt(*Buffer, 0, 3), TEXT(' '));
+			TestEqual(TEXT("C shifted right"), GetCharAt(*Buffer, 0, 4), TEXT('C'));
+		});
+
+		It("DCH: should delete characters at cursor and shift left", [this]()
+		{
+			ParseString(*Parser, TEXT("ABCDEF"));
+			ParseString(*Parser, TEXT("\x1B[1;3H")); // cursor at (0, 2)
+			ParseString(*Parser, TEXT("\x1B[2P"));    // DCH 2
+
+			TestEqual(TEXT("A preserved"), GetCharAt(*Buffer, 0, 0), TEXT('A'));
+			TestEqual(TEXT("B preserved"), GetCharAt(*Buffer, 0, 1), TEXT('B'));
+			TestEqual(TEXT("E shifted left"), GetCharAt(*Buffer, 0, 2), TEXT('E'));
+			TestEqual(TEXT("F shifted left"), GetCharAt(*Buffer, 0, 3), TEXT('F'));
+		});
+
+		It("IL: should insert blank lines at cursor row", [this]()
+		{
+			ParseString(*Parser, TEXT("Line0"));
+			ParseString(*Parser, TEXT("\x1B[2;1H")); // cursor at row 1
+			ParseString(*Parser, TEXT("Line1"));
+			ParseString(*Parser, TEXT("\x1B[2;1H")); // back to row 1
+			ParseString(*Parser, TEXT("\x1B[1L"));    // IL 1
+
+			// Row 1 should now be blank; original Line1 should be pushed down.
+			TestEqual(TEXT("blank inserted line"), GetCharAt(*Buffer, 1, 0), TEXT(' '));
+			TestEqual(TEXT("Line1 pushed down"), GetCharAt(*Buffer, 2, 0), TEXT('L'));
+		});
+
+		It("DL: should delete lines at cursor row", [this]()
+		{
+			ParseString(*Parser, TEXT("Line0"));
+			ParseString(*Parser, TEXT("\x1B[2;1H")); // row 1
+			ParseString(*Parser, TEXT("Line1"));
+			ParseString(*Parser, TEXT("\x1B[3;1H")); // row 2
+			ParseString(*Parser, TEXT("Line2"));
+			ParseString(*Parser, TEXT("\x1B[2;1H")); // back to row 1
+			ParseString(*Parser, TEXT("\x1B[1M"));    // DL 1
+
+			// Row 1 should now have Line2 content.
+			TestEqual(TEXT("Line2 moved up"), GetCharAt(*Buffer, 1, 0), TEXT('L'));
+			TestEqual(TEXT("Line2 content"), GetCharAt(*Buffer, 1, 4), TEXT('2'));
+		});
+	});
+
+	// ── E. Scroll ──
+
+	Describe("Scroll", [this]()
+	{
+		It("SU: should scroll up by Ps lines", [this]()
+		{
+			ParseString(*Parser, TEXT("Row0"));
+			ParseString(*Parser, TEXT("\x1B[2;1H"));
+			ParseString(*Parser, TEXT("Row1"));
+			ParseString(*Parser, TEXT("\x1B[3;1H"));
+			ParseString(*Parser, TEXT("Row2"));
+
+			ParseString(*Parser, TEXT("\x1B[1S")); // SU 1
+
+			// Row0 should have scrolled off; Row1 should now be at viewport row 0.
+			TestEqual(TEXT("Row1 at top"), GetCharAt(*Buffer, 0, 3), TEXT('1'));
+		});
+
+		It("SD: should scroll down by Ps lines", [this]()
+		{
+			ParseString(*Parser, TEXT("Row0"));
+			ParseString(*Parser, TEXT("\x1B[2;1H"));
+			ParseString(*Parser, TEXT("Row1"));
+
+			ParseString(*Parser, TEXT("\x1B[1T")); // SD 1
+
+			// Row0 should now be at viewport row 1; viewport row 0 should be blank.
+			TestEqual(TEXT("blank at top"), GetCharAt(*Buffer, 0, 0), TEXT(' '));
+			TestEqual(TEXT("Row0 pushed down"), GetCharAt(*Buffer, 1, 3), TEXT('0'));
+		});
+	});
+
+	// ── F. SGR  - Select Graphic Rendition ──
+
+	Describe("SGR", [this]()
+	{
+		It("SGR 0: should reset all attributes", [this]()
+		{
+			ParseString(*Parser, TEXT("\x1B[1m"));  // Bold
+			ParseString(*Parser, TEXT("\x1B[31m")); // Red foreground
+			ParseString(*Parser, TEXT("\x1B[0m"));  // Reset
+
+			TestEqual(TEXT("attributes reset"), Parser->CurrentAttributes, static_cast<uint8>(ETerminalAttribute::None));
+			TestEqual(TEXT("foreground reset"), Parser->CurrentForeground, FColor(212, 212, 212));
+			TestEqual(TEXT("background reset"), Parser->CurrentBackground, FColor(30, 30, 30));
+		});
+
+		It("SGR 1: should set bold", [this]()
+		{
+			ParseString(*Parser, TEXT("\x1B[1m"));
+			TestTrue(TEXT("bold set"), (Parser->CurrentAttributes & ETerminalAttribute::Bold) != 0);
+		});
+
+		It("SGR 2: should set dim", [this]()
+		{
+			ParseString(*Parser, TEXT("\x1B[2m"));
+			TestTrue(TEXT("dim set"), (Parser->CurrentAttributes & ETerminalAttribute::Dim) != 0);
+		});
+
+		It("SGR 3: should set italic", [this]()
+		{
+			ParseString(*Parser, TEXT("\x1B[3m"));
+			TestTrue(TEXT("italic set"), (Parser->CurrentAttributes & ETerminalAttribute::Italic) != 0);
+		});
+
+		It("SGR 4: should set underline", [this]()
+		{
+			ParseString(*Parser, TEXT("\x1B[4m"));
+			TestTrue(TEXT("underline set"), (Parser->CurrentAttributes & ETerminalAttribute::Underline) != 0);
+		});
+
+		It("SGR 7: should set inverse", [this]()
+		{
+			ParseString(*Parser, TEXT("\x1B[7m"));
+			TestTrue(TEXT("inverse set"), (Parser->CurrentAttributes & ETerminalAttribute::Inverse) != 0);
+		});
+
+		It("SGR 9: should set strikethrough", [this]()
+		{
+			ParseString(*Parser, TEXT("\x1B[9m"));
+			TestTrue(TEXT("strikethrough set"), (Parser->CurrentAttributes & ETerminalAttribute::Strikethrough) != 0);
+		});
+
+		It("SGR 22: should clear bold and dim", [this]()
+		{
+			ParseString(*Parser, TEXT("\x1B[1;2m"));  // Bold + Dim
+			ParseString(*Parser, TEXT("\x1B[22m"));   // Normal intensity
+			TestTrue(TEXT("bold cleared"), (Parser->CurrentAttributes & ETerminalAttribute::Bold) == 0);
+			TestTrue(TEXT("dim cleared"), (Parser->CurrentAttributes & ETerminalAttribute::Dim) == 0);
+		});
+
+		It("SGR 23: should clear italic", [this]()
+		{
+			ParseString(*Parser, TEXT("\x1B[3m"));
+			ParseString(*Parser, TEXT("\x1B[23m"));
+			TestTrue(TEXT("italic cleared"), (Parser->CurrentAttributes & ETerminalAttribute::Italic) == 0);
+		});
+
+		It("SGR 24: should clear underline", [this]()
+		{
+			ParseString(*Parser, TEXT("\x1B[4m"));
+			ParseString(*Parser, TEXT("\x1B[24m"));
+			TestTrue(TEXT("underline cleared"), (Parser->CurrentAttributes & ETerminalAttribute::Underline) == 0);
+		});
+
+		It("SGR 27: should clear inverse", [this]()
+		{
+			ParseString(*Parser, TEXT("\x1B[7m"));
+			ParseString(*Parser, TEXT("\x1B[27m"));
+			TestTrue(TEXT("inverse cleared"), (Parser->CurrentAttributes & ETerminalAttribute::Inverse) == 0);
+		});
+
+		It("SGR 29: should clear strikethrough", [this]()
+		{
+			ParseString(*Parser, TEXT("\x1B[9m"));
+			ParseString(*Parser, TEXT("\x1B[29m"));
+			TestTrue(TEXT("strikethrough cleared"), (Parser->CurrentAttributes & ETerminalAttribute::Strikethrough) == 0);
+		});
+
+		It("SGR 31: should set red foreground", [this]()
+		{
+			ParseString(*Parser, TEXT("\x1B[31m"));
+			TestEqual(TEXT("foreground is red"), Parser->CurrentForeground, Parser->AnsiPalette[1]);
+		});
+
+		It("SGR 42: should set green background", [this]()
+		{
+			ParseString(*Parser, TEXT("\x1B[42m"));
+			TestEqual(TEXT("background is green"), Parser->CurrentBackground, Parser->AnsiPalette[2]);
+		});
+
+		It("SGR 91: should set bright red foreground", [this]()
+		{
+			ParseString(*Parser, TEXT("\x1B[91m"));
+			TestEqual(TEXT("foreground is bright red"), Parser->CurrentForeground, Parser->AnsiPalette[9]);
+		});
+
+		It("SGR 102: should set bright green background", [this]()
+		{
+			ParseString(*Parser, TEXT("\x1B[102m"));
+			TestEqual(TEXT("background is bright green"), Parser->CurrentBackground, Parser->AnsiPalette[10]);
+		});
+
+		It("SGR 39: should reset to default foreground", [this]()
+		{
+			ParseString(*Parser, TEXT("\x1B[31m"));  // Red
+			ParseString(*Parser, TEXT("\x1B[39m"));  // Default
+			TestEqual(TEXT("foreground default"), Parser->CurrentForeground, FColor(212, 212, 212));
+		});
+
+		It("SGR 49: should reset to default background", [this]()
+		{
+			ParseString(*Parser, TEXT("\x1B[42m"));  // Green bg
+			ParseString(*Parser, TEXT("\x1B[49m"));  // Default
+			TestEqual(TEXT("background default"), Parser->CurrentBackground, FColor(30, 30, 30));
+		});
+
+		It("SGR 38;5;N: should set 256-color foreground", [this]()
+		{
+			// Color 196 in 6x6x6 cube: index 196 = 16 + 180 => r=5, g=0, b=0
+			// r = (196 - 16) / 36 = 5, g = ((196 - 16) % 36) / 6 = 0, b = (196 - 16) % 6 = 0
+			// RGB = (255, 0, 0)
+			ParseString(*Parser, TEXT("\x1B[38;5;196m"));
+			TestEqual(TEXT("256-color R"), Parser->CurrentForeground.R, static_cast<uint8>(255));
+			TestEqual(TEXT("256-color G"), Parser->CurrentForeground.G, static_cast<uint8>(0));
+			TestEqual(TEXT("256-color B"), Parser->CurrentForeground.B, static_cast<uint8>(0));
+		});
+
+		It("SGR 38;5;N: should handle grayscale ramp", [this]()
+		{
+			// Grayscale index 240: 8 + (240-232)*10 = 88 => RGB(88,88,88)
+			ParseString(*Parser, TEXT("\x1B[38;5;240m"));
+			TestEqual(TEXT("grayscale R"), Parser->CurrentForeground.R, static_cast<uint8>(88));
+			TestEqual(TEXT("grayscale G"), Parser->CurrentForeground.G, static_cast<uint8>(88));
+			TestEqual(TEXT("grayscale B"), Parser->CurrentForeground.B, static_cast<uint8>(88));
+		});
+
+		It("SGR 38;2;R;G;B: should set truecolor foreground", [this]()
+		{
+			ParseString(*Parser, TEXT("\x1B[38;2;100;150;200m"));
+			TestEqual(TEXT("truecolor R"), Parser->CurrentForeground.R, static_cast<uint8>(100));
+			TestEqual(TEXT("truecolor G"), Parser->CurrentForeground.G, static_cast<uint8>(150));
+			TestEqual(TEXT("truecolor B"), Parser->CurrentForeground.B, static_cast<uint8>(200));
+		});
+
+		It("should handle multiple SGR params in one sequence", [this]()
+		{
+			ParseString(*Parser, TEXT("\x1B[1;31m")); // Bold + Red
+			TestTrue(TEXT("bold set"), (Parser->CurrentAttributes & ETerminalAttribute::Bold) != 0);
+			TestEqual(TEXT("foreground is red"), Parser->CurrentForeground, Parser->AnsiPalette[1]);
+		});
+	});
+
+	// ── G. Scroll Region  - DECSTBM ──
+
+	Describe("DECSTBM Scroll Region", [this]()
+	{
+		It("should set scroll region with CSI Ps;Ps r", [this]()
+		{
+			ParseString(*Parser, TEXT("\x1B[5;20r")); // Set region rows 5-20 (1-based)
+			TestEqual(TEXT("scroll region top"), Parser->ScrollRegionTop, 4);
+			TestEqual(TEXT("scroll region bottom"), Parser->ScrollRegionBottom, 19);
+		});
+
+		It("should scroll only within region on LF at bottom", [this]()
+		{
+			// Set scroll region to rows 2-4 (1-based), i.e. viewport rows 1-3.
+			ParseString(*Parser, TEXT("\x1B[2;4r"));
+
+			// Write text in the region.
+			ParseString(*Parser, TEXT("\x1B[2;1H")); // row 2 (viewport 1)
+			ParseString(*Parser, TEXT("Line1"));
+			ParseString(*Parser, TEXT("\x1B[3;1H")); // row 3 (viewport 2)
+			ParseString(*Parser, TEXT("Line2"));
+			ParseString(*Parser, TEXT("\x1B[4;1H")); // row 4 (viewport 3)  - bottom of region
+			ParseString(*Parser, TEXT("Line3"));
+
+			// Write something on row 1 (above region) that should not scroll.
+			ParseString(*Parser, TEXT("\x1B[1;1H"));
+			ParseString(*Parser, TEXT("Above"));
+
+			// Move cursor to bottom of region and LF.
+			ParseString(*Parser, TEXT("\x1B[4;1H"));
+			ParseBytes(*Parser, { 0x0A }); // LF
+
+			// "Above" on row 0 should be untouched.
+			TestEqual(TEXT("row above region preserved"), GetCharAt(*Buffer, 0, 0), TEXT('A'));
+
+			// Line1 should have scrolled off; Line2 should now be at viewport row 1.
+			TestEqual(TEXT("Line2 scrolled up"), GetCharAt(*Buffer, 1, 4), TEXT('2'));
+		});
+
+		It("should restore full-screen scrolling with CSI r (no params)", [this]()
+		{
+			ParseString(*Parser, TEXT("\x1B[5;20r"));
+			ParseString(*Parser, TEXT("\x1B[r")); // Reset
+			// No-param DECSTBM defaults to 1;ViewportRows, setting the region to the full viewport.
+			TestEqual(TEXT("scroll region top reset"), Parser->ScrollRegionTop, 0);
+			TestEqual(TEXT("scroll region bottom reset"), Parser->ScrollRegionBottom, 23);
+		});
+	});
+
+	// ── H. DEC Private Modes  - DECSET/DECRST ──
+
+	Describe("DEC Private Modes", [this]()
+	{
+		It("DECSET ?1 / DECRST ?1: application/normal cursor keys", [this]()
+		{
+			ParseString(*Parser, TEXT("\x1B[?1h")); // DECSET
+			TestTrue(TEXT("application cursor keys on"), Parser->bApplicationCursorKeys);
+			ParseString(*Parser, TEXT("\x1B[?1l")); // DECRST
+			TestFalse(TEXT("application cursor keys off"), Parser->bApplicationCursorKeys);
+		});
+
+		It("DECSET ?25 / DECRST ?25: show/hide cursor", [this]()
+		{
+			ParseString(*Parser, TEXT("\x1B[?25l")); // Hide
+			TestFalse(TEXT("cursor hidden"), Parser->bCursorVisible);
+			ParseString(*Parser, TEXT("\x1B[?25h")); // Show
+			TestTrue(TEXT("cursor shown"), Parser->bCursorVisible);
+		});
+
+		It("DECSET ?2004 / DECRST ?2004: bracketed paste mode", [this]()
+		{
+			ParseString(*Parser, TEXT("\x1B[?2004h")); // Enable
+			TestTrue(TEXT("bracketed paste on"), Parser->bBracketedPaste);
+			ParseString(*Parser, TEXT("\x1B[?2004l")); // Disable
+			TestFalse(TEXT("bracketed paste off"), Parser->bBracketedPaste);
+		});
+
+		It("DECSET ?1049 / DECRST ?1049: alternate screen buffer", [this]()
+		{
+			ParseString(*Parser, TEXT("Main"));
+			ParseString(*Parser, TEXT("\x1B[?1049h")); // Switch to alt buffer
+			TestTrue(TEXT("alt buffer active"), Buffer->IsAlternateBufferActive());
+
+			ParseString(*Parser, TEXT("\x1B[?1049l")); // Restore main buffer
+			TestFalse(TEXT("main buffer restored"), Buffer->IsAlternateBufferActive());
+			TestEqual(TEXT("main content preserved"), GetCharAt(*Buffer, 0, 0), TEXT('M'));
+		});
+
+		It("DECSET ?2026 / DECRST ?2026: synchronized output", [this]()
+		{
+			ParseString(*Parser, TEXT("\x1B[?2026h"));
+			TestTrue(TEXT("synchronized output on"), Parser->bSynchronizedOutput);
+			ParseString(*Parser, TEXT("\x1B[?2026l"));
+			TestFalse(TEXT("synchronized output off"), Parser->bSynchronizedOutput);
+		});
+
+		It("DECSET ?1000 / DECRST ?1000: normal mouse tracking", [this]()
+		{
+			ParseString(*Parser, TEXT("\x1B[?1000h"));
+			TestTrue(TEXT("normal tracking on"), Parser->MouseTrackingMode == FVTParser::EMouseTrackingMode::Normal);
+			ParseString(*Parser, TEXT("\x1B[?1000l"));
+			TestTrue(TEXT("tracking off"), Parser->MouseTrackingMode == FVTParser::EMouseTrackingMode::None);
+		});
+
+		It("DECSET ?1002 / DECRST ?1002: button-event mouse tracking", [this]()
+		{
+			ParseString(*Parser, TEXT("\x1B[?1002h"));
+			TestTrue(TEXT("button-event tracking on"), Parser->MouseTrackingMode == FVTParser::EMouseTrackingMode::ButtonEvent);
+			ParseString(*Parser, TEXT("\x1B[?1002l"));
+			TestTrue(TEXT("tracking off"), Parser->MouseTrackingMode == FVTParser::EMouseTrackingMode::None);
+		});
+
+		It("DECSET ?1003 / DECRST ?1003: any-event mouse tracking", [this]()
+		{
+			ParseString(*Parser, TEXT("\x1B[?1003h"));
+			TestTrue(TEXT("any-event tracking on"), Parser->MouseTrackingMode == FVTParser::EMouseTrackingMode::Any);
+			ParseString(*Parser, TEXT("\x1B[?1003l"));
+			TestTrue(TEXT("tracking off"), Parser->MouseTrackingMode == FVTParser::EMouseTrackingMode::None);
+		});
+
+		It("DECSET ?1006 / DECRST ?1006: SGR extended mouse encoding", [this]()
+		{
+			ParseString(*Parser, TEXT("\x1B[?1006h"));
+			TestTrue(TEXT("SGR encoding on"), Parser->bSGRMouseEncoding);
+			ParseString(*Parser, TEXT("\x1B[?1006l"));
+			TestFalse(TEXT("SGR encoding off"), Parser->bSGRMouseEncoding);
+		});
+
+		It("should set multiple modes in one sequence", [this]()
+		{
+			ParseString(*Parser, TEXT("\x1B[?1000;1006h"));
+			TestTrue(TEXT("normal tracking on"), Parser->MouseTrackingMode == FVTParser::EMouseTrackingMode::Normal);
+			TestTrue(TEXT("SGR encoding on"), Parser->bSGRMouseEncoding);
+			ParseString(*Parser, TEXT("\x1B[?1000;1006l"));
+			TestTrue(TEXT("tracking off"), Parser->MouseTrackingMode == FVTParser::EMouseTrackingMode::None);
+			TestFalse(TEXT("SGR encoding off"), Parser->bSGRMouseEncoding);
+		});
+	});
+
+	// ── I. OSC  - Operating System Command ──
+
+	Describe("OSC", [this]()
+	{
+		It("OSC 0: should set window title (BEL terminated)", [this]()
+		{
+			ParseString(*Parser, TEXT("\x1B]0;My Title\x07"));
+			TestEqual(TEXT("window title"), Parser->WindowTitle, TEXT("My Title"));
+		});
+
+		It("OSC 2: should set window title (BEL terminated)", [this]()
+		{
+			ParseString(*Parser, TEXT("\x1B]2;Another Title\x07"));
+			TestEqual(TEXT("window title"), Parser->WindowTitle, TEXT("Another Title"));
+		});
+
+		It("OSC: should accept ST terminator (ESC backslash)", [this]()
+		{
+			ParseString(*Parser, TEXT("\x1B]0;ST Title\x1B\\"));
+			TestEqual(TEXT("window title via ST"), Parser->WindowTitle, TEXT("ST Title"));
+		});
+	});
+
+	// ── J. Device Status Report / Device Attributes ──
+
+	Describe("Device Status Report", [this]()
+	{
+		It("DSR 5n: should report status OK", [this]()
+		{
+			ParseString(*Parser, TEXT("\x1B[5n"));
+			FString Response(reinterpret_cast<const char*>(Parser->ResponseBuffer.GetData()), Parser->ResponseBuffer.Num());
+			TestTrue(TEXT("response contains 0n"), Response.Contains(TEXT("0n")));
+		});
+
+		It("DSR 6n: should report cursor position (1-based)", [this]()
+		{
+			Buffer->Cursor.Row = 4;
+			Buffer->Cursor.Column = 9;
+			ParseString(*Parser, TEXT("\x1B[6n"));
+			FString Response(reinterpret_cast<const char*>(Parser->ResponseBuffer.GetData()), Parser->ResponseBuffer.Num());
+			// Should report row 5, column 10 (1-based).
+			TestTrue(TEXT("response contains 5;10R"), Response.Contains(TEXT("5;10R")));
+		});
+
+		It("Primary DA: should return a response", [this]()
+		{
+			ParseString(*Parser, TEXT("\x1B[c"));
+			TestTrue(TEXT("response buffer not empty"), Parser->ResponseBuffer.Num() > 0);
+		});
+	});
+
+	// ── K. Escape Sequences ──
+
+	Describe("Escape Sequences", [this]()
+	{
+		It("ESC 7 / ESC 8: should save and restore cursor position", [this]()
+		{
+			Buffer->Cursor.Row = 5;
+			Buffer->Cursor.Column = 10;
+			ParseString(*Parser, TEXT("\x1B\x37")); // ESC 7 (DECSC)
+
+			Buffer->Cursor.Row = 0;
+			Buffer->Cursor.Column = 0;
+			ParseString(*Parser, TEXT("\x1B\x38")); // ESC 8 (DECRC)
+
+			TestEqual(TEXT("restored row"), Buffer->Cursor.Row, 5);
+			TestEqual(TEXT("restored column"), Buffer->Cursor.Column, 10);
+		});
+
+		It("ESC H: should set tab stop at current column", [this]()
+		{
+			// Clear all tab stops first.
+			ParseString(*Parser, TEXT("\x1B[3g"));
+
+			// Set a tab stop at column 5.
+			Buffer->Cursor.Column = 5;
+			ParseString(*Parser, TEXT("\x1BH")); // HTS
+
+			// Move to column 0 and tab.
+			Buffer->Cursor.Column = 0;
+			ParseBytes(*Parser, { 0x09 }); // HT
+			TestEqual(TEXT("tab to custom stop"), Buffer->Cursor.Column, 5);
+		});
+
+		It("CSI 0g: should clear tab stop at current column", [this]()
+		{
+			// Default tab stops at every 8 columns.
+			Buffer->Cursor.Column = 8;
+			ParseString(*Parser, TEXT("\x1B[0g")); // Clear tab at column 8
+
+			Buffer->Cursor.Column = 0;
+			ParseBytes(*Parser, { 0x09 }); // HT
+			// Should skip column 8 and go to 16.
+			TestEqual(TEXT("tab skips cleared stop"), Buffer->Cursor.Column, 16);
+		});
+
+		It("CSI 3g: should clear all tab stops", [this]()
+		{
+			ParseString(*Parser, TEXT("\x1B[3g"));
+
+			Buffer->Cursor.Column = 0;
+			ParseBytes(*Parser, { 0x09 }); // HT
+			// No tab stops  - should go to last column.
+			TestEqual(TEXT("tab goes to end"), Buffer->Cursor.Column, 79);
+		});
+	});
+
+	// ── L. UTF-8 Decoding ──
+
+	Describe("UTF-8 Decoding", [this]()
+	{
+		It("should decode 2-byte sequence (U+00E9, e with acute)", [this]()
+		{
+			const uint8 Bytes[] = { 0xC3, 0xA9 };
+			ParseBytes(*Parser, Bytes);
+			TestEqual(TEXT("character is e-acute"), GetCharAt(*Buffer, 0, 0), static_cast<TCHAR>(0x00E9));
+		});
+
+		It("should decode 3-byte sequence (U+2713, check mark)", [this]()
+		{
+			const uint8 Bytes[] = { 0xE2, 0x9C, 0x93 };
+			ParseBytes(*Parser, Bytes);
+			TestEqual(TEXT("character is check mark"), GetCharAt(*Buffer, 0, 0), static_cast<TCHAR>(0x2713));
+		});
+
+		It("should decode 4-byte sequence (U+1F600, emoji)", [this]()
+		{
+			const uint8 Bytes[] = { 0xF0, 0x9F, 0x98, 0x80 };
+			ParseBytes(*Parser, Bytes);
+			// 4-byte codepoints on platforms with 16-bit TCHAR get stored as surrogate pairs
+			// or implementation-defined. Just verify no crash and cursor advanced.
+			TestTrue(TEXT("cursor advanced"), Buffer->Cursor.Column >= 1);
+		});
+
+		It("should recover from invalid continuation byte", [this]()
+		{
+			// Start a 2-byte sequence, then send an ASCII char instead of continuation.
+			const uint8 Bytes[] = { 0xC3, 0x41 }; // 0xC3 expects continuation, 0x41 = 'A'
+			ParseBytes(*Parser, Bytes);
+			// Should recover and print 'A'.
+			TestEqual(TEXT("recovered with A"), GetCharAt(*Buffer, 0, 0), TEXT('A'));
+		});
+
+		It("should reject overlong encoding (C0 80 = overlong NUL)", [this]()
+		{
+			const uint8 Bytes[] = { 0xC0, 0x80, 0x41 }; // Overlong NUL, then 'A'
+			ParseBytes(*Parser, Bytes);
+			// Overlong should be rejected. 'A' should be at column 0.
+			TestEqual(TEXT("overlong rejected, A printed"), GetCharAt(*Buffer, 0, 0), TEXT('A'));
+		});
+	});
+
+	// ── M. Parser State Machine Resilience ──
+
+	Describe("Parser State Machine Resilience", [this]()
+	{
+		It("CAN mid-sequence should return to ground state", [this]()
+		{
+			// Start a CSI sequence, then CAN, then print.
+			const uint8 Bytes[] = { 0x1B, 0x5B, 0x31, 0x18, 0x41 }; // ESC [ 1 CAN A
+			ParseBytes(*Parser, Bytes);
+			// 'A' should be printed normally.
+			TestEqual(TEXT("A printed after CAN"), GetCharAt(*Buffer, 0, 0), TEXT('A'));
+		});
+
+		It("SUB mid-sequence should return to ground state", [this]()
+		{
+			const uint8 Bytes[] = { 0x1B, 0x5B, 0x31, 0x1A, 0x42 }; // ESC [ 1 SUB B
+			ParseBytes(*Parser, Bytes);
+			TestEqual(TEXT("B printed after SUB"), GetCharAt(*Buffer, 0, 0), TEXT('B'));
+		});
+
+		It("ESC mid-sequence should start new escape sequence", [this]()
+		{
+			// ESC [ 1 ESC [ 2 J  - the first CSI is abandoned, second is ED 2.
+			ParseString(*Parser, TEXT("Hello"));
+			ParseString(*Parser, TEXT("\x1B[1\x1B[2J"));
+
+			// ED 2 should have cleared the display.
+			TestEqual(TEXT("display cleared"), GetCharAt(*Buffer, 0, 0), TEXT(' '));
+		});
+
+		It("malformed CSI with private marker mid-params should be ignored", [this]()
+		{
+			// CSI 1 ? 2 m  - a '?' appearing after the first parameter position
+			// triggers CSI ignore per the Williams spec (0x3C-0x3F after first pos).
+			ParseString(*Parser, TEXT("\x1B[1?2m"));
+			ParseString(*Parser, TEXT("X"));
+			// After the ignored sequence, 'X' should print normally.
+			TestEqual(TEXT("X printed after ignored CSI"), GetCharAt(*Buffer, 0, 0), TEXT('X'));
+		});
+	});
+}
+
+#endif // WITH_DEV_AUTOMATION_TESTS
diff --git a/Engine/Plugins/Experimental/Terminal/Source/TerminalTests/TerminalTests.Build.cs b/Engine/Plugins/Experimental/Terminal/Source/TerminalTests/TerminalTests.Build.cs
new file mode 100644
--- /dev/null
+++ b/Engine/Plugins/Experimental/Terminal/Source/TerminalTests/TerminalTests.Build.cs
@@ -0,0 +1,29 @@
+// Copyright Epic Games, Inc. All Rights Reserved.
+
+using UnrealBuildTool;
+
+public class TerminalTests : ModuleRules
+{
+	public TerminalTests(ReadOnlyTargetRules Target) : base(Target)
+	{
+		PCHUsage = ModuleRules.PCHUsageMode.UseExplicitOrSharedPCHs;
+
+		PublicDependencyModuleNames.AddRange(
+			new string[]
+			{
+				"Core"
+			});
+
+		PrivateDependencyModuleNames.AddRange(
+			new string[]
+			{
+				"CoreUObject",
+				"Engine",
+				"InputCore",
+				"Slate",
+				"SlateCore",
+				"Terminal",
+				"UnrealEd"
+			});
+	}
+}
diff --git a/Engine/Plugins/Experimental/Terminal/Terminal.uplugin b/Engine/Plugins/Experimental/Terminal/Terminal.uplugin
new file mode 100644
--- /dev/null
+++ b/Engine/Plugins/Experimental/Terminal/Terminal.uplugin
@@ -0,0 +1,30 @@
+{
+	"FileVersion": 3,
+	"Version": 1,
+	"VersionName": "1.0",
+	"FriendlyName": "Terminal",
+	"Description": "Native Slate terminal emulator.",
+	"NoRedist": true,
+	"EditorOnly": true,
+	"CreatedBy": "",
+	"CreatedByURL": "",
+	"DocsURL": "",
+	"MarketplaceURL": "",
+	"CanContainContent": false,
+	"IsBetaVersion": false,
+	"IsExperimentalVersion": true,
+	"Installed": false,
+	"EnabledByDefault": false,
+	"Modules": [
+		{
+			"Name": "Terminal",
+			"Type": "Editor",
+			"LoadingPhase": "Default"
+		},
+		{
+			"Name": "TerminalTests",
+			"Type": "Editor",
+			"LoadingPhase": "Default"
+		}
+	]
+}
