[Terminal] Add `StartupCommands` setting to execute commands on new terminal window creation.

- Add a `StartupCommands` (`TArray<FString>`) property to `UTerminalSettings` that allows defining commands to run automatically when a new terminal window is created. Commands are written to the PTY input pipe right after session creation and execute in order once the shell is ready.

#Synced-CL 52520731

diff --git a/Engine/Plugins/Experimental/Terminal/Source/Terminal/Private/STerminal.cpp b/Engine/Plugins/Experimental/Terminal/Source/Terminal/Private/STerminal.cpp
--- a/Engine/Plugins/Experimental/Terminal/Source/Terminal/Private/STerminal.cpp
+++ b/Engine/Plugins/Experimental/Terminal/Source/Terminal/Private/STerminal.cpp
@@ -1249,6 +1249,16 @@ void STerminal::InitializeTerminal()
 			const FUtf8String ErrorUtf8(TEXT("\r\nError: Failed to create terminal session. Check the output log for details.\r\n"));
 			Parser.Parse(reinterpret_cast<const uint8*>(*ErrorUtf8), ErrorUtf8.Len());
 		}
+		else if (Settings)
+		{
+			for (const FString& StartupCommand : Settings->StartupCommands)
+			{
+				if (!StartupCommand.IsEmpty())
+				{
+					ExecuteCommand(StartupCommand);
+				}
+			}
+		}
 	}
 
 	bInitialized = true;
diff --git a/Engine/Plugins/Experimental/Terminal/Source/Terminal/Public/TerminalSettings.h b/Engine/Plugins/Experimental/Terminal/Source/Terminal/Public/TerminalSettings.h
--- a/Engine/Plugins/Experimental/Terminal/Source/Terminal/Public/TerminalSettings.h
+++ b/Engine/Plugins/Experimental/Terminal/Source/Terminal/Public/TerminalSettings.h
@@ -32,4 +32,8 @@ public:
 	/** Name of the color scheme to use. Must match a JSON file in Config/ColorSchemes/. */
 	UPROPERTY(Config, EditAnywhere, Category="Terminal")
 	FString ColorSchemeName = TEXT("Default");
+
+	/** Commands to execute automatically when a new terminal window is created. Each entry is sent as a separate command. */
+	UPROPERTY(Config, EditAnywhere, Category="Terminal")
+	TArray<FString> StartupCommands;
 };
